Added new simple NNE runtime NNERuntimeBasicCpu to unify custom inference implementations across MLDeformer and LearningAgents

diff --git a/Engine/Plugins/Experimental/NNERuntimeBasicCpu/NNERuntimeBasicCpu.uplugin b/Engine/Plugins/Experimental/NNERuntimeBasicCpu/NNERuntimeBasicCpu.uplugin
new file mode 100644
--- /dev/null
+++ b/Engine/Plugins/Experimental/NNERuntimeBasicCpu/NNERuntimeBasicCpu.uplugin
@@ -0,0 +1,26 @@
+{
+	"FileVersion": 3,
+	"Version": 0.1,
+	"VersionName": "0.1",
+	"FriendlyName": "NNERuntimeBasicCpu",
+	"Description": "Performant, cross-platform, CPU runtime for the NNE plugin that supports basic models.",
+	"Category": "Experimental",
+	"CreatedBy": "Epic Games, Inc.",
+	"CreatedByURL": "https://epicgames.com",
+	"DocsURL": "",
+	"MarketplaceURL": "",
+	"SupportURL": "",
+	"EnabledByDefault": false,
+	"CanContainContent": true,
+	"IsBetaVersion": false,
+	"IsExperimentalVersion": true,
+	"Installed": false,
+	"Modules": [
+		{
+			"Name": "NNERuntimeBasicCpu",
+			"Type": "RuntimeAndProgram",
+			"LoadingPhase": "Default"
+		}
+	],
+	"Plugins": []
+}
\ No newline at end of file
diff --git a/Engine/Plugins/Experimental/NNERuntimeBasicCpu/Source/NNERuntimeBasicCpu/NNERuntimeBasicCpu.Build.cs b/Engine/Plugins/Experimental/NNERuntimeBasicCpu/Source/NNERuntimeBasicCpu/NNERuntimeBasicCpu.Build.cs
new file mode 100644
--- /dev/null
+++ b/Engine/Plugins/Experimental/NNERuntimeBasicCpu/Source/NNERuntimeBasicCpu/NNERuntimeBasicCpu.Build.cs
@@ -0,0 +1,29 @@
+// Copyright Epic Games, Inc. All Rights Reserved.
+
+using UnrealBuildTool;
+
+public class NNERuntimeBasicCpu : ModuleRules
+{
+	public NNERuntimeBasicCpu( ReadOnlyTargetRules Target ) : base( Target )
+	{
+		PCHUsage = PCHUsageMode.UseExplicitOrSharedPCHs;
+
+		PublicDependencyModuleNames.AddRange
+		(
+			new string[] 
+			{
+				"Core",
+				"Engine",
+				"NNE",
+			}
+		);
+
+		PrivateDependencyModuleNames.AddRange
+		(
+			new string[] 
+			{
+				"CoreUObject",
+			}
+		);
+	}
+}
diff --git a/Engine/Plugins/Experimental/NNERuntimeBasicCpu/Source/NNERuntimeBasicCpu/Private/NNERuntimeBasicCpu.cpp b/Engine/Plugins/Experimental/NNERuntimeBasicCpu/Source/NNERuntimeBasicCpu/Private/NNERuntimeBasicCpu.cpp
new file mode 100644
--- /dev/null
+++ b/Engine/Plugins/Experimental/NNERuntimeBasicCpu/Source/NNERuntimeBasicCpu/Private/NNERuntimeBasicCpu.cpp
@@ -0,0 +1,133 @@
+// Copyright Epic Games, Inc. All Rights Reserved.
+
+#include "NNERuntimeBasicCpu.h"
+#include "NNERuntimeBasicCpuModel.h"
+#include "NNEModelData.h"
+
+// We ask for the memory to be aligned to 64 bytes since this is the
+// largest alignment we ask for inside the ModelData for pointers to
+// various bits of data.
+const uint32 UNNERuntimeBasicCpuImpl::Alignment = 64;
+
+bool UNNERuntimeBasicCpuImpl::CanCreateModelData(FString FileType, TConstArrayView<uint8> FileData, FGuid FileId, const ITargetPlatform* TargetPlatform) const
+{
+	if (FileType.Compare("ubnne", ESearchCase::IgnoreCase) != 0)
+	{
+		return false;
+	}
+
+	// We require at least a magic number and version number
+	if (FileData.Num() < 2 * sizeof(uint32))
+	{
+		return false;
+	}
+
+	// Check magic number valid
+	const uint32* FileMagicNumber = (const uint32*)&FileData[0 * sizeof(uint32)];
+	if (*FileMagicNumber != UE::NNE::RuntimeBasic::FModelCPU::ModelMagicNumber)
+	{
+		return false;
+	}
+
+	// Check version number valid
+	const uint32* FileVersionNumber = (const uint32*)&FileData[1 * sizeof(uint32)];
+	if (*FileVersionNumber != UE::NNE::RuntimeBasic::FModelCPU::ModelVersionNumber)
+	{
+		return false;
+	}
+
+	return true;
+}
+
+TSharedPtr<UE::NNE::FSharedModelData> UNNERuntimeBasicCpuImpl::CreateModelData(FString FileType, TConstArrayView<uint8> FileData, FGuid FileId, const ITargetPlatform* TargetPlatform)
+{
+	if (!CanCreateModelData(FileType, FileData, FileId, TargetPlatform))
+	{
+		return nullptr;
+	}
+
+	// The model data is the same as what is stored in the file however we need to
+	// make an owned copy and also ensure it is aligned since we will be creating 
+	// pointers to inside.
+	uint8* ModelData = (uint8*)FMemory::Malloc(FileData.Num(), Alignment);
+	check(ModelData);
+	FMemory::Memcpy(ModelData, FileData.GetData(), FileData.Num());
+
+	return MakeShared<UE::NNE::FSharedModelData>(FSharedBuffer::TakeOwnership(ModelData, FileData.Num(), FMemory::Free), Alignment);
+}
+
+FString UNNERuntimeBasicCpuImpl::GetModelDataIdentifier(FString FileType, TConstArrayView<uint8> FileData, FGuid FileId, const ITargetPlatform* TargetPlatform)
+{
+	return FileId.ToString(EGuidFormats::Digits) + "-" + FString::FromInt(UE::NNE::RuntimeBasic::FModelCPU::ModelMagicNumber);
+}
+
+bool UNNERuntimeBasicCpuImpl::CanCreateModelCPU(TObjectPtr<UNNEModelData> ModelData) const
+{
+	check(ModelData != nullptr);
+
+	TSharedPtr<UE::NNE::FSharedModelData> SharedData = ModelData->GetModelData(GetRuntimeName());
+
+	if (!SharedData.IsValid())
+	{
+		return false;
+	}
+
+	TConstArrayView<uint8> Data = SharedData->GetView();
+
+	// We require at least a magic number and version number
+	if (Data.Num() < 2 * sizeof(uint32))
+	{
+		return false;
+	}
+
+	// Check magic number valid
+	const uint32* FileMagicNumber = (const uint32*)&Data[0 * sizeof(uint32)];
+	if (*FileMagicNumber != UE::NNE::RuntimeBasic::FModelCPU::ModelMagicNumber)
+	{
+		return false;
+	}
+
+	// Check version number valid
+	const uint32* FileVersionNumber = (const uint32*)&Data[1 * sizeof(uint32)];
+	if (*FileVersionNumber != UE::NNE::RuntimeBasic::FModelCPU::ModelVersionNumber)
+	{
+		return false;
+	}
+
+	return true;
+}
+
+TSharedPtr<UE::NNE::IModelCPU> UNNERuntimeBasicCpuImpl::CreateModelCPU(TObjectPtr<UNNEModelData> ModelData)
+{
+	check(ModelData != nullptr);
+
+	if (!CanCreateModelCPU(ModelData))
+	{
+		return nullptr;
+	}
+
+	// Create Model Object
+	TSharedPtr<UE::NNE::RuntimeBasic::FModelCPU> Model = MakeShared<UE::NNE::RuntimeBasic::FModelCPU>();
+	
+	// Store weak pointer to self
+	Model->WeakThis = Model.ToWeakPtr();
+
+	// Create Model Data
+	Model->ModelData = ModelData->GetModelData(GetRuntimeName());
+	check(Model->ModelData != nullptr);
+
+	// Check Model Data Alignment
+	checkf(((UPTRINT)(Model->ModelData->GetView().GetData()) % Alignment) == 0,
+		TEXT("Model data must be aligned to %i bytes."), Alignment);
+
+	// Load from model data
+	uint64 Offset = 0;
+	if (Model->SerializationLoad(Offset, Model->ModelData->GetView()))
+	{
+		return Model;
+	}
+	else
+	{
+		return nullptr;
+	}
+}
diff --git a/Engine/Plugins/Experimental/NNERuntimeBasicCpu/Source/NNERuntimeBasicCpu/Private/NNERuntimeBasicCpu.ispc b/Engine/Plugins/Experimental/NNERuntimeBasicCpu/Source/NNERuntimeBasicCpu/Private/NNERuntimeBasicCpu.ispc
new file mode 100644
--- /dev/null
+++ b/Engine/Plugins/Experimental/NNERuntimeBasicCpu/Source/NNERuntimeBasicCpu/Private/NNERuntimeBasicCpu.ispc
@@ -0,0 +1,220 @@
+// Copyright Epic Games, Inc. All Rights Reserved.
+
+export void NNERuntimeBasicCPUNormalizeLayer(
+	uniform float Output[],
+	uniform const float Input[],
+	uniform const float Mean[],
+	uniform const float Std[],
+	uniform const int BatchNum,
+	uniform const int InputOutputNum,
+	uniform const int OutputBufferSize,
+	uniform const int InputBufferSize)
+{
+	foreach(BatchIdx = 0 ... BatchNum, Idx = 0 ... InputOutputNum)
+	{
+		Output[BatchIdx * OutputBufferSize + Idx] = (Input[BatchIdx * InputBufferSize + Idx] - Mean[Idx]) / Std[Idx];
+	}
+}
+
+export void NNERuntimeBasicCPUDenormalizeLayer(
+	uniform float Output[],
+	uniform const float Input[],
+	uniform const float Mean[],
+	uniform const float Std[],
+	uniform const int BatchNum,
+	uniform const int InputOutputNum,
+	uniform const int OutputBufferSize,
+	uniform const int InputBufferSize)
+{
+	foreach(BatchIdx = 0 ... BatchNum, Idx = 0 ... InputOutputNum)
+	{
+		Output[BatchIdx * OutputBufferSize + Idx] = (Input[BatchIdx * InputBufferSize + Idx] * Std[Idx]) + Mean[Idx];
+	}
+}
+
+export void NNERuntimeBasicCPULinearLayer(
+	uniform float Output[],
+	uniform const float Input[],
+	uniform const float Weights[],
+	uniform const float Biases[],
+	uniform const int BatchNum,
+	uniform const int RowNum,
+	uniform const int ColNum,
+	uniform const int OutputBufferSize,
+	uniform const int InputBufferSize)
+{
+	foreach(BatchIdx = 0 ... BatchNum, ColIdx = 0 ... ColNum)
+	{
+		Output[BatchIdx * OutputBufferSize + ColIdx] = Biases[ColIdx];
+	}
+
+	for (uniform int BatchIdx = 0; BatchIdx < BatchNum; BatchIdx++)
+	{
+		for (uniform int RowIdx = 0; RowIdx < RowNum; RowIdx++)
+		{
+			uniform const float Value = Input[BatchIdx * InputBufferSize + RowIdx];
+
+			if (Value != 0.0)
+			{
+				foreach(ColIdx = 0 ... ColNum)
+				{
+					Output[BatchIdx * OutputBufferSize + ColIdx] += Value * Weights[RowIdx * ColNum + ColIdx];
+				}
+			}
+		}
+	}
+}
+
+export void NNERuntimeBasicCPUCompressedLinearLayer(
+	uniform float Output[],
+	uniform const float Input[],
+	uniform const uint16 Weights[],
+	uniform const float WeightOffsets[],
+	uniform const float WeightScales[],
+	uniform const float Biases[],
+	uniform const int BatchNum,
+	uniform const int RowNum,
+	uniform const int ColNum,
+	uniform const int OutputBufferSize,
+	uniform const int InputBufferSize)
+{
+	foreach(BatchIdx = 0 ... BatchNum, ColIdx = 0 ... ColNum)
+	{
+		Output[BatchIdx * OutputBufferSize + ColIdx] = Biases[ColIdx];
+	}
+
+	for (uniform int BatchIdx = 0; BatchIdx < BatchNum; BatchIdx++)
+	{
+		for (uniform int RowIdx = 0; RowIdx < RowNum; RowIdx++)
+		{
+			uniform const float Value = Input[BatchIdx * InputBufferSize + RowIdx];
+
+			if (Value != 0.0)
+			{
+				uniform const float Offset = WeightOffsets[RowIdx];
+				uniform const float Scale = WeightScales[RowIdx];
+
+				foreach(ColIdx = 0 ... ColNum)
+				{
+					Output[BatchIdx * OutputBufferSize + ColIdx] += Value * ((Scale * ((float)Weights[RowIdx * ColNum + ColIdx])) + Offset);
+				}
+			}
+		}
+	}
+}
+
+export void NNERuntimeBasicCPUMultiLinearLayer(
+	uniform float Output[],
+	uniform const float Input[],
+	uniform const float Weights[],
+	uniform const float Biases[],
+	uniform const int BatchNum,
+	uniform const int BlockNum,
+	uniform const int RowNum,
+	uniform const int ColNum,
+	uniform const int OutputBufferSize,
+	uniform const int InputBufferSize)
+{
+	foreach(BatchIdx = 0 ... BatchNum, BlockIdx = 0 ... BlockNum, ColIdx = 0 ... ColNum)
+	{
+		Output[BatchIdx * OutputBufferSize + BlockIdx * ColNum + ColIdx] = Biases[BlockIdx * ColNum + ColIdx];
+	}
+
+	for (uniform int BatchIdx = 0; BatchIdx < BatchNum; BatchIdx++)
+	{
+		for (uniform int BlockIdx = 0; BlockIdx < BlockNum; BlockIdx++)
+		{
+			for (uniform int RowIdx = 0; RowIdx < RowNum; RowIdx++)
+			{
+				uniform const float Value = Input[BatchIdx * InputBufferSize + BlockIdx * RowNum + RowIdx];
+
+				if (Value != 0.0)
+				{
+					foreach (ColIdx = 0 ... ColNum)
+					{
+						Output[BatchIdx * OutputBufferSize + BlockIdx * ColNum + ColIdx] += Value * Weights[BlockIdx * RowNum * ColNum + RowIdx * ColNum + ColIdx];
+					}
+				}
+			}
+		}
+	}
+}
+
+export void NNERuntimeBasicCPUActivationLinear(
+	uniform float Output[],
+	uniform const float Input[],
+	uniform const int BatchNum,
+	uniform const int Num,
+	uniform const int OutputBufferSize,
+	uniform const int InputBufferSize)
+{
+	foreach(BatchIdx = 0 ... BatchNum, Idx = 0 ... Num)
+	{
+		Output[BatchIdx * OutputBufferSize + Idx] = Input[BatchIdx * InputBufferSize + Idx];
+	}
+}
+
+export void NNERuntimeBasicCPUActivationReLU(
+	uniform float Output[],
+	uniform const float Input[],
+	uniform const int BatchNum,
+	uniform const int Num,
+	uniform const int OutputBufferSize,
+	uniform const int InputBufferSize)
+{
+	foreach(BatchIdx = 0 ... BatchNum, Idx = 0 ... Num)
+	{
+		Output[BatchIdx * OutputBufferSize + Idx] = max(Input[BatchIdx * InputBufferSize + Idx], 0.0f);
+	}
+}
+
+export void NNERuntimeBasicCPUActivationELU(
+	uniform float Output[],
+	uniform const float Input[],
+	uniform const int BatchNum,
+	uniform const int Num,
+	uniform const int OutputBufferSize,
+	uniform const int InputBufferSize)
+{
+	foreach(BatchIdx = 0 ... BatchNum, Idx = 0 ... Num)
+	{
+		const float Value = Input[BatchIdx * InputBufferSize + Idx];
+		Output[BatchIdx * OutputBufferSize + Idx] = select(Value > 0.0f, Value, exp(Value) - 1.0f);
+	}
+}
+
+inline float NNERuntimeBasicCPUTanH(const float X, const float Min = -10.0f, const float Max = 10.0f)
+{
+	const float Exp2X = exp(2.0f * clamp(X, Min, Max));
+	return (Exp2X - 1) / (Exp2X + 1);
+}
+
+export void NNERuntimeBasicCPUActivationTanH(
+	uniform float Output[],
+	uniform const float Input[],
+	uniform const int BatchNum,
+	uniform const int Num,
+	uniform const int OutputBufferSize,
+	uniform const int InputBufferSize)
+{
+	foreach(BatchIdx = 0 ... BatchNum, Idx = 0 ... Num)
+	{
+		Output[BatchIdx * OutputBufferSize + Idx] = NNERuntimeBasicCPUTanH(Input[BatchIdx * InputBufferSize + Idx]);
+	}
+}
+
+export void NNERuntimeBasicCPUActivationPReLU(
+	uniform float Output[],
+	uniform const float Input[],
+	uniform const float Alpha[],
+	uniform const int BatchNum,
+	uniform const int Num,
+	uniform const int OutputBufferSize,
+	uniform const int InputBufferSize)
+{
+	foreach(BatchIdx = 0 ... BatchNum, Idx = 0 ... Num)
+	{
+		const float Value = Input[BatchIdx * InputBufferSize + Idx];
+		Output[BatchIdx * OutputBufferSize + Idx] = select(Value > 0.0f, Value, Alpha[Idx] * Value);
+	}
+}
diff --git a/Engine/Plugins/Experimental/NNERuntimeBasicCpu/Source/NNERuntimeBasicCpu/Private/NNERuntimeBasicCpuModel.cpp b/Engine/Plugins/Experimental/NNERuntimeBasicCpu/Source/NNERuntimeBasicCpu/Private/NNERuntimeBasicCpuModel.cpp
new file mode 100644
--- /dev/null
+++ b/Engine/Plugins/Experimental/NNERuntimeBasicCpu/Source/NNERuntimeBasicCpu/Private/NNERuntimeBasicCpuModel.cpp
@@ -0,0 +1,1534 @@
+// Copyright Epic Games, Inc. All Rights Reserved.
+
+#include "NNERuntimeBasicCpuModel.h"
+
+#include "NNE.h"
+#include "NNERuntimeBasicCpu.h"
+
+#define NNE_RUNTIME_BASIC_ENABLE_ISPC INTEL_ISPC
+//#define NNE_RUNTIME_BASIC_ENABLE_ISPC 0
+
+#if NNE_RUNTIME_BASIC_ENABLE_ISPC
+#include "NNERuntimeBasicCpu.ispc.generated.h"
+#endif
+
+namespace UE::NNE::RuntimeBasic
+{
+	namespace Private
+	{
+		namespace Serialization
+		{
+			//--------------------------------------------------------------------------
+
+			static inline void Align(uint64& InOutOffset, const uint32 Alignment)
+			{
+				InOutOffset = ((InOutOffset + Alignment - 1) / Alignment) * Alignment;
+			}
+
+			//--------------------------------------------------------------------------
+
+			static inline void Size(uint64& InOutOffset, const uint32& In)
+			{
+				Align(InOutOffset, sizeof(uint32));
+				InOutOffset += sizeof(uint32);
+			}
+
+			static inline void Size(uint64& InOutOffset, const float& In)
+			{
+				Align(InOutOffset, sizeof(float));
+				InOutOffset += sizeof(float);
+			}
+
+			static inline void Size(uint64& InOutOffset, const TConstArrayView<float> In)
+			{
+				Align(InOutOffset, 64);
+				InOutOffset += In.Num() * sizeof(float);
+			}
+
+			static inline void Size(uint64& InOutOffset, const TConstArrayView<uint16> In)
+			{
+				Align(InOutOffset, 64);
+				InOutOffset += In.Num() * sizeof(uint16);
+			}
+
+			static inline void Size(uint64& InOutOffset, const TSharedPtr<ILayer>& InLayer);
+
+			//--------------------------------------------------------------------------
+
+			static inline void Load(uint64& InOutOffset, uint32& Out, TConstArrayView<uint8> Data)
+			{
+				Align(InOutOffset, sizeof(uint32));
+				Out = *((uint32*)&Data[InOutOffset]);
+				InOutOffset += sizeof(uint32);
+			}
+
+			static inline void Load(uint64& InOutOffset, float& Out, TConstArrayView<uint8> Data)
+			{
+				Align(InOutOffset, sizeof(float));
+				Out = *((float*)&Data[InOutOffset]);
+				InOutOffset += sizeof(float);
+			}
+
+			static inline void Load(uint64& InOutOffset, TConstArrayView<float>& Out, TConstArrayView<uint8> Data, int32 Size)
+			{
+				Align(InOutOffset, 64);
+				Out = MakeArrayView<float>((float*)&Data[InOutOffset], Size);
+				InOutOffset += Size * sizeof(float);
+			}
+
+			static inline void Load(uint64& InOutOffset, TConstArrayView<uint16>& Out, TConstArrayView<uint8> Data, int32 Size)
+			{
+				Align(InOutOffset, 64);
+				Out = MakeArrayView<uint16>((uint16*)&Data[InOutOffset], Size);
+				InOutOffset += Size * sizeof(uint16);
+			}
+
+			static inline void Load(uint64& InOutOffset, TSharedPtr<ILayer>& OutLayer, TConstArrayView<uint8> Data);
+
+			//--------------------------------------------------------------------------
+
+			static inline void Save(uint64& InOutOffset, const uint32 In, TArrayView<uint8> Data)
+			{
+				Align(InOutOffset, sizeof(uint32));
+				*((uint32*)&Data[InOutOffset]) = In;
+				InOutOffset += sizeof(uint32);
+			}
+
+			static inline void Save(uint64& InOutOffset, const float In, TArrayView<uint8> Data)
+			{
+				Align(InOutOffset, sizeof(float));
+				*((float*)&Data[InOutOffset]) = In;
+				InOutOffset += sizeof(float);
+			}
+
+			static inline void Save(uint64& InOutOffset, TConstArrayView<float> In, TArrayView<uint8> Data)
+			{
+				Align(InOutOffset, 64);
+				FMemory::Memcpy(&Data[InOutOffset], In.GetData(), In.Num() * sizeof(float));
+				InOutOffset += In.Num() * sizeof(float);
+			}
+
+			static inline void Save(uint64& InOutOffset, TConstArrayView<uint16> In, TArrayView<uint8> Data)
+			{
+				Align(InOutOffset, 64);
+				FMemory::Memcpy(&Data[InOutOffset], In.GetData(), In.Num() * sizeof(uint16));
+				InOutOffset += In.Num() * sizeof(uint16);
+			}
+
+			static inline void Save(uint64& InOutOffset, const TSharedPtr<ILayer>& InLayer, TArrayView<uint8> Data);
+		}
+
+		//--------------------------------------------------------------------------
+
+		/** Layer Type Id - this should match what is given in nne_runtime_basic_cpu.py  */
+		enum class ELayerType : uint32
+		{
+			Invalid = 0,
+			Sequence = 1,
+			Normalize = 2,
+			Denormalize = 3,
+			Linear = 4,
+			CompressedLinear = 5,
+			MultiLinear = 6,
+			ReLU = 7,
+			ELU = 8,
+			TanH = 9,
+			PReLU = 10,
+		};
+
+		//--------------------------------------------------------------------------
+
+		/**
+		 * Interface for a Layer Instance - the data required for performing inference for a layer.
+		 */
+		struct ILayerInstance
+		{
+			/** Virtual destructor */
+			virtual ~ILayerInstance() = default;
+
+			/** Indicate to this layer instance what the batchsize is going to be when performing inference. */
+			virtual void SetBatchSize(const uint32 BatchSize) = 0;
+		};
+
+		/**
+		 * Interface for a Layer - the network parameter data required for a layer.
+		 */
+		struct ILayer
+		{
+			/** Virtual destructor */
+			virtual ~ILayer() = default;
+
+			/** Create the instance data required for this type of layer. */
+			virtual TSharedPtr<ILayerInstance> MakeInstance() const { return nullptr; };
+
+			/** Get the layer type. */
+			virtual ELayerType GetLayerType() const = 0;
+
+			/** Get the size of the input vector. */
+			virtual uint32 GetInputSize() const = 0;
+
+			/** Get the size of the output vector. */
+			virtual uint32 GetOutputSize() const = 0;
+
+			/** Compute the size required to serialize this layer by growing InOutOffset. */
+			virtual void SerializationSize(uint64& InOutOffset) const = 0;
+
+			/** Load this layer from the buffer at the given offset. */
+			virtual void SerializationLoad(uint64& InOutOffset, TConstArrayView<uint8> Data) = 0;
+
+			/** Save this layer from the buffer at the given offset. */
+			virtual void SerializationSave(uint64& InOutOffset, TArrayView<uint8> Data) const = 0;
+
+			/**
+			 * Evaluate this layer.
+			 * 
+			 * Note that the input and output buffers may not match the exact dimensions expected by this layer - the
+			 * size dimension can be over-allocated to allow for buffer re-use.
+			 *
+			 * @param Output			The output buffer of size (BatchSize, OutputBufferSize).
+			 * @param Input				The input buffer of size (BatchSize, InputBufferSize).
+			 * @param Instance			The instance data for this layer.
+			 * @param BatchSize			The batchsize of the inputs and outputs
+			 * @param OutputBufferSize	The size of the output buffer vector for each item in the batch
+			 * @param InputBufferSize	The size of the input buffer vector for each item in the batch
+			 */
+			virtual void Evaluate(
+				TArrayView<float> Output,
+				TConstArrayView<float> Input,
+				ILayerInstance* Instance,
+				uint32 BatchSize,
+				uint32 OutputBufferSize,
+				uint32 InputBufferSize) = 0;
+		};
+
+		//--------------------------------------------------------------------------
+
+		struct FSequenceLayer;
+
+		struct FSequenceLayerInstance : public ILayerInstance
+		{
+			FSequenceLayerInstance(const FSequenceLayer& InSequenceLayer);
+
+			virtual void SetBatchSize(const uint32 BatchSize) override final;
+
+			const FSequenceLayer& SequenceLayer;
+			uint32 ActivationBufferSize = 0;
+			TArray<TSharedPtr<ILayerInstance>, TInlineAllocator<32>> Instances;
+			TArray<float, TInlineAllocator<512>> ActivationBufferFront;
+			TArray<float, TInlineAllocator<512>> ActivationBufferBack;
+		};
+
+		//--------------------------------------------------------------------------
+
+		struct FSequenceLayer : public ILayer
+		{
+			virtual TSharedPtr<ILayerInstance> MakeInstance() const { return MakeShared<FSequenceLayerInstance>(*this); };
+			virtual ELayerType GetLayerType() const override final { return ELayerType::Sequence; }
+			virtual uint32 GetInputSize() const override final { return Layers.Num() > 0 ? Layers[0]->GetInputSize() : 0; }
+			virtual uint32 GetOutputSize() const override final { return Layers.Num() > 0 ? Layers.Last()->GetOutputSize() : 0; }
+
+			virtual void SerializationSize(uint64& InOutOffset) const override final
+			{
+				Serialization::Size(InOutOffset, (uint32)Layers.Num());
+				for (const TSharedPtr<ILayer>& Layer : Layers)
+				{
+					Serialization::Size(InOutOffset, Layer);
+				}
+			}
+
+			virtual void SerializationLoad(uint64& InOutOffset, TConstArrayView<uint8> Data) override final
+			{
+				uint32 LayerNum = 0;
+				Serialization::Load(InOutOffset, LayerNum, Data);
+				Layers.SetNumZeroed(LayerNum);
+
+				for (uint32 LayerIdx = 0; LayerIdx < LayerNum; LayerIdx++)
+				{
+					Serialization::Load(InOutOffset, Layers[LayerIdx], Data);
+				}
+			}
+
+			virtual void SerializationSave(uint64& InOutOffset, TArrayView<uint8> Data) const override final
+			{
+				uint32 LayerNum = Layers.Num();
+				Serialization::Save(InOutOffset, LayerNum, Data);
+				for (uint32 LayerIdx = 0; LayerIdx < LayerNum; LayerIdx++)
+				{
+					Serialization::Save(InOutOffset, Layers[LayerIdx], Data);
+				}
+			}
+
+			virtual void Evaluate(
+				TArrayView<float> Output,
+				TConstArrayView<float> Input,
+				ILayerInstance* Instance,
+				uint32 BatchSize,
+				uint32 OutputBufferSize,
+				uint32 InputBufferSize) override final
+			{
+				TRACE_CPUPROFILER_EVENT_SCOPE(NNE::RuntimeBasic::Private::FSequenceLayer::Evaluate);
+
+				const uint32 LayerNum = Layers.Num();
+				check(LayerNum > 0);
+
+				FSequenceLayerInstance* SequenceInstance = StaticCast<FSequenceLayerInstance*>(Instance);
+				check(SequenceInstance);
+
+				// If we just have one layer then evaluate layer directly without using intermediate storage
+
+				if (LayerNum == 1)
+				{
+					Layers[0]->Evaluate(
+						Output,
+						Input,
+						SequenceInstance->Instances[0].Get(),
+						BatchSize,
+						OutputBufferSize,
+						InputBufferSize);
+
+					return;
+				}
+
+				// Otherwise evaluate first layer from input into activation buffer
+
+				Layers[0]->Evaluate(
+					SequenceInstance->ActivationBufferFront,
+					Input,
+					SequenceInstance->Instances[0].Get(),
+					BatchSize,
+					SequenceInstance->ActivationBufferSize,
+					InputBufferSize);
+
+				// Evaluate intermediate layers using front and back buffers
+
+				for (uint32 LayerIdx = 1; LayerIdx < LayerNum - 1; LayerIdx++)
+				{
+					TConstArrayView<float> LayerInput = LayerIdx % 2 == 0 ? 
+						SequenceInstance->ActivationBufferBack : 
+						SequenceInstance->ActivationBufferFront;
+
+					TArrayView<float> LayerOutput = LayerIdx % 2 == 0 ? 
+						SequenceInstance->ActivationBufferFront : 
+						SequenceInstance->ActivationBufferBack;
+
+					Layers[LayerIdx]->Evaluate(
+						LayerOutput,
+						LayerInput,
+						SequenceInstance->Instances[LayerIdx].Get(),
+						BatchSize,
+						SequenceInstance->ActivationBufferSize,
+						SequenceInstance->ActivationBufferSize);
+				}
+
+				// Evaluate final layer from activation buffer into output
+
+				TConstArrayView<float> FinalLayerInput = LayerNum % 2 == 0 ? 
+					SequenceInstance->ActivationBufferFront : 
+					SequenceInstance->ActivationBufferBack;
+
+				Layers.Last()->Evaluate(
+					Output,
+					FinalLayerInput,
+					SequenceInstance->Instances.Last().Get(),
+					BatchSize,
+					OutputBufferSize,
+					SequenceInstance->ActivationBufferSize);
+			}
+
+			TArray<TSharedPtr<ILayer>, TInlineAllocator<32>> Layers;
+		};
+
+		//--------------------------------------------------------------------------
+
+		FSequenceLayerInstance::FSequenceLayerInstance(const FSequenceLayer& InSequenceLayer) 
+			: SequenceLayer(InSequenceLayer)
+		{
+			const uint32 LayerNum = SequenceLayer.Layers.Num();
+			Instances.Init(nullptr, LayerNum);
+
+			for (uint32 LayerIdx = 0; LayerIdx < LayerNum; LayerIdx++)
+			{
+				Instances[LayerIdx] = SequenceLayer.Layers[LayerIdx]->MakeInstance();
+			}
+
+			// Compute the largest intermediate size used
+
+			ActivationBufferSize = SequenceLayer.Layers[0]->GetInputSize();
+			for (uint32 LayerIdx = 1; LayerIdx < LayerNum; LayerIdx++)
+			{
+				ActivationBufferSize = FMath::Max(ActivationBufferSize, SequenceLayer.Layers[LayerIdx]->GetOutputSize());
+			}
+		}
+
+		void FSequenceLayerInstance::SetBatchSize(const uint32 BatchSize)
+		{
+			TRACE_CPUPROFILER_EVENT_SCOPE(NNE::RuntimeBasic::Private::FSequenceLayerInstance::SetBatchSize);
+
+			// Propagate call to sub-layer instances
+
+			for (const TSharedPtr<ILayerInstance>& Instance : Instances)
+			{
+				// Most layers don't allocate instance data and so we need to check for nullptr
+				if (Instance)
+				{
+					Instance->SetBatchSize(BatchSize);
+				}
+			}
+
+			// Allocate front and back buffers to maximum size. Don't shrink to avoid re-allocation 
+			// when smaller batches are requested.
+
+			ActivationBufferFront.SetNumUninitialized(BatchSize * ActivationBufferSize, false);
+			ActivationBufferBack.SetNumUninitialized(BatchSize * ActivationBufferSize, false);
+		}
+
+		//--------------------------------------------------------------------------
+
+		static inline void NormalizeLayer(
+			float* RESTRICT Output,
+			const float* RESTRICT Input,
+			const float* RESTRICT Mean,
+			const float* RESTRICT Std,
+			const uint32 BatchNum,
+			const uint32 InputOutputNum,
+			const uint32 OutputBufferSize,
+			const uint32 InputBufferSize)
+		{
+			for (uint32 BatchIdx = 0; BatchIdx < BatchNum; BatchIdx++)
+			{
+				for (uint32 Idx = 0; Idx < InputOutputNum; Idx++)
+				{
+					Output[BatchIdx * OutputBufferSize + Idx] = (Input[BatchIdx * InputBufferSize + Idx] - Mean[Idx]) / Std[Idx];
+				}
+			}
+		}
+
+		struct FNormalizeLayer : public ILayer
+		{
+			virtual ELayerType GetLayerType() const override final { return ELayerType::Normalize; }
+			virtual uint32 GetInputSize() const override final { return InputOutputSize; }
+			virtual uint32 GetOutputSize() const override final { return InputOutputSize; }
+
+			virtual void SerializationSize(uint64& InOutOffset) const override final
+			{
+				Serialization::Size(InOutOffset, InputOutputSize);
+				Serialization::Size(InOutOffset, Mean);
+				Serialization::Size(InOutOffset, Std);
+			}
+
+			virtual void SerializationLoad(uint64& InOutOffset, TConstArrayView<uint8> Data) override final
+			{
+				Serialization::Load(InOutOffset, InputOutputSize, Data);
+				Serialization::Load(InOutOffset, Mean, Data, InputOutputSize);
+				Serialization::Load(InOutOffset, Std, Data, InputOutputSize);
+			}
+
+			virtual void SerializationSave(uint64& InOutOffset, TArrayView<uint8> Data) const override final
+			{
+				Serialization::Save(InOutOffset, InputOutputSize, Data);
+				Serialization::Save(InOutOffset, Mean, Data);
+				Serialization::Save(InOutOffset, Std, Data);
+			}
+
+			virtual void Evaluate(
+				TArrayView<float> Output,
+				TConstArrayView<float> Input,
+				ILayerInstance* Instance,
+				uint32 BatchSize,
+				uint32 OutputBufferSize,
+				uint32 InputBufferSize) override final
+			{
+				TRACE_CPUPROFILER_EVENT_SCOPE(NNE::RuntimeBasic::Private::FNormalizeLayer::Evaluate);
+
+#if NNE_RUNTIME_BASIC_ENABLE_ISPC
+				ispc::NNERuntimeBasicCPUNormalizeLayer(
+					Output.GetData(),
+					Input.GetData(),
+					Mean.GetData(),
+					Std.GetData(),
+					BatchSize,
+					InputOutputSize,
+					OutputBufferSize,
+					InputBufferSize);
+#else
+				NormalizeLayer(
+					Output.GetData(),
+					Input.GetData(),
+					Mean.GetData(),
+					Std.GetData(),
+					BatchSize,
+					InputOutputSize,
+					OutputBufferSize,
+					InputBufferSize);
+#endif
+			}
+
+			uint32 InputOutputSize = 0;
+			TConstArrayView<float> Mean;
+			TConstArrayView<float> Std;
+		};
+
+		//--------------------------------------------------------------------------
+
+		static inline void DenormalizeLayer(
+			float* RESTRICT Output,
+			const float* RESTRICT Input,
+			const float* RESTRICT Mean,
+			const float* RESTRICT Std,
+			const uint32 BatchNum,
+			const uint32 InputOutputNum,
+			const uint32 OutputBufferSize,
+			const uint32 InputBufferSize)
+		{
+			for (uint32 BatchIdx = 0; BatchIdx < BatchNum; BatchIdx++)
+			{
+				for (uint32 Idx = 0; Idx < InputOutputNum; Idx++)
+				{
+					Output[BatchIdx * OutputBufferSize + Idx] = (Input[BatchIdx * InputBufferSize + Idx] * Std[Idx]) + Mean[Idx];
+				}
+			}
+		}
+
+		struct FDenormalizeLayer : public ILayer
+		{
+			virtual ELayerType GetLayerType() const override final { return ELayerType::Denormalize; }
+			virtual uint32 GetInputSize() const override final { return InputOutputSize; }
+			virtual uint32 GetOutputSize() const override final { return InputOutputSize; }
+
+			virtual void SerializationSize(uint64& InOutOffset) const override final
+			{
+				Serialization::Size(InOutOffset, InputOutputSize);
+				Serialization::Size(InOutOffset, Mean);
+				Serialization::Size(InOutOffset, Std);
+			}
+
+			virtual void SerializationLoad(uint64& InOutOffset, TConstArrayView<uint8> Data) override final
+			{
+				Serialization::Load(InOutOffset, InputOutputSize, Data);
+				Serialization::Load(InOutOffset, Mean, Data, InputOutputSize);
+				Serialization::Load(InOutOffset, Std, Data, InputOutputSize);
+			}
+
+			virtual void SerializationSave(uint64& InOutOffset, TArrayView<uint8> Data) const override final
+			{
+				Serialization::Save(InOutOffset, InputOutputSize, Data);
+				Serialization::Save(InOutOffset, Mean, Data);
+				Serialization::Save(InOutOffset, Std, Data);
+			}
+
+			virtual void Evaluate(
+				TArrayView<float> Output,
+				TConstArrayView<float> Input,
+				ILayerInstance* Instance,
+				uint32 BatchSize,
+				uint32 OutputBufferSize,
+				uint32 InputBufferSize) override final
+			{
+				TRACE_CPUPROFILER_EVENT_SCOPE(NNE::RuntimeBasic::Private::FDenormalizeLayer::Evaluate);
+
+#if NNE_RUNTIME_BASIC_ENABLE_ISPC
+				ispc::NNERuntimeBasicCPUDenormalizeLayer(
+					Output.GetData(),
+					Input.GetData(),
+					Mean.GetData(),
+					Std.GetData(),
+					BatchSize,
+					InputOutputSize,
+					OutputBufferSize,
+					InputBufferSize);
+#else
+				DenormalizeLayer(
+					Output.GetData(),
+					Input.GetData(),
+					Mean.GetData(),
+					Std.GetData(),
+					BatchSize,
+					InputOutputSize,
+					OutputBufferSize,
+					InputBufferSize);
+#endif
+			}
+
+			uint32 InputOutputSize = 0;
+			TConstArrayView<float> Mean;
+			TConstArrayView<float> Std;
+		};
+
+		//--------------------------------------------------------------------------
+
+		static inline void LinearLayer(
+			float* RESTRICT Output,
+			const float* RESTRICT Input,
+			const float* RESTRICT Weights,
+			const float* RESTRICT Biases,
+			const uint32 BatchNum,
+			const uint32 RowNum,
+			const uint32 ColNum,
+			const uint32 OutputBufferSize,
+			const uint32 InputBufferSize)
+		{
+			for (uint32 BatchIdx = 0; BatchIdx < BatchNum; BatchIdx++)
+			{
+				for (uint32 ColIdx = 0; ColIdx < ColNum; ColIdx++)
+				{
+					Output[BatchIdx * OutputBufferSize + ColIdx] = Biases[ColIdx];
+				}
+			}
+
+			for (uint32 BatchIdx = 0; BatchIdx < BatchNum; BatchIdx++)
+			{
+				for (uint32 RowIdx = 0; RowIdx < RowNum; RowIdx++)
+				{
+					const float Value = Input[BatchIdx * InputBufferSize + RowIdx];
+
+					if (Value != 0.0)
+					{
+						for (uint32 ColIdx = 0; ColIdx < ColNum; ColIdx++)
+						{
+							Output[BatchIdx * OutputBufferSize + ColIdx] += Value * Weights[RowIdx * ColNum + ColIdx];
+						}
+					}
+				}
+			}
+		}
+
+		struct FLinearLayer : public ILayer
+		{
+			virtual ELayerType GetLayerType() const override final { return ELayerType::Linear; }
+			virtual uint32 GetInputSize() const override final { return InputSize; }
+			virtual uint32 GetOutputSize() const override final { return OutputSize; }
+
+			virtual void SerializationSize(uint64& InOutOffset) const override final
+			{
+				Serialization::Size(InOutOffset, InputSize);
+				Serialization::Size(InOutOffset, OutputSize);
+				Serialization::Size(InOutOffset, Biases);
+				Serialization::Size(InOutOffset, Weights);
+			}
+
+			virtual void SerializationLoad(uint64& InOutOffset, TConstArrayView<uint8> Data) override final
+			{
+				Serialization::Load(InOutOffset, InputSize, Data);
+				Serialization::Load(InOutOffset, OutputSize, Data);
+				Serialization::Load(InOutOffset, Biases, Data, OutputSize);
+				Serialization::Load(InOutOffset, Weights, Data, InputSize * OutputSize);
+			}
+
+			virtual void SerializationSave(uint64& InOutOffset, TArrayView<uint8> Data) const override final
+			{
+				Serialization::Save(InOutOffset, InputSize, Data);
+				Serialization::Save(InOutOffset, OutputSize, Data);
+				Serialization::Save(InOutOffset, Biases, Data);
+				Serialization::Save(InOutOffset, Weights, Data);
+			}
+
+			virtual void Evaluate(
+				TArrayView<float> Output,
+				TConstArrayView<float> Input,
+				ILayerInstance* Instance,
+				uint32 BatchSize,
+				uint32 OutputBufferSize,
+				uint32 InputBufferSize) override final
+			{
+				TRACE_CPUPROFILER_EVENT_SCOPE(NNE::RuntimeBasic::Private::FLinearLayer::Evaluate);
+
+#if NNE_RUNTIME_BASIC_ENABLE_ISPC
+				ispc::NNERuntimeBasicCPULinearLayer(
+					Output.GetData(),
+					Input.GetData(),
+					Weights.GetData(),
+					Biases.GetData(),
+					BatchSize,
+					InputSize,
+					OutputSize,
+					OutputBufferSize,
+					InputBufferSize);
+#else
+				LinearLayer(
+					Output.GetData(),
+					Input.GetData(),
+					Weights.GetData(),
+					Biases.GetData(),
+					BatchSize,
+					InputSize,
+					OutputSize,
+					OutputBufferSize,
+					InputBufferSize);
+#endif
+			}
+
+			uint32 InputSize = 0;
+			uint32 OutputSize = 0;
+			TConstArrayView<float> Biases;
+			TConstArrayView<float> Weights;
+		};
+
+		//--------------------------------------------------------------------------
+
+		static inline void CompressedLinearLayer(
+			float* RESTRICT Output,
+			const float* RESTRICT Input,
+			const uint16* RESTRICT Weights,
+			const float* RESTRICT WeightOffsets,
+			const float* RESTRICT WeightScales,
+			const float* RESTRICT Biases,
+			const uint32 BatchNum,
+			const uint32 RowNum,
+			const uint32 ColNum,
+			const uint32 OutputBufferSize,
+			const uint32 InputBufferSize)
+		{
+			for (uint32 BatchIdx = 0; BatchIdx < BatchNum; BatchIdx++)
+			{
+				for (uint32 ColIdx = 0; ColIdx < ColNum; ColIdx++)
+				{
+					Output[BatchIdx * OutputBufferSize + ColIdx] = Biases[ColIdx];
+				}
+			}
+
+			for (uint32 BatchIdx = 0; BatchIdx < BatchNum; BatchIdx++)
+			{
+				for (uint32 RowIdx = 0; RowIdx < RowNum; RowIdx++)
+				{
+					const float Value = Input[BatchIdx * InputBufferSize + RowIdx];
+
+					if (Value != 0.0)
+					{
+						const float Offset = WeightOffsets[RowIdx];
+						const float Scales = WeightScales[RowIdx];
+
+						for (uint32 ColIdx = 0; ColIdx < ColNum; ColIdx++)
+						{
+							Output[BatchIdx * OutputBufferSize + ColIdx] += Value * ((Scales * ((float)Weights[RowIdx * ColNum + ColIdx])) + Offset);
+						}
+					}
+				}
+			}
+		}
+
+		struct FCompressedLinearLayer : public ILayer
+		{
+			virtual ELayerType GetLayerType() const override final { return ELayerType::CompressedLinear; }
+			virtual uint32 GetInputSize() const override final { return InputSize; }
+			virtual uint32 GetOutputSize() const override final { return OutputSize; }
+
+			virtual void SerializationSize(uint64& InOutOffset) const override final
+			{
+				Serialization::Size(InOutOffset, InputSize);
+				Serialization::Size(InOutOffset, OutputSize);
+				Serialization::Size(InOutOffset, WeightOffsets);
+				Serialization::Size(InOutOffset, WeightScales);
+				Serialization::Size(InOutOffset, Biases);
+				Serialization::Size(InOutOffset, Weights);
+			}
+
+			virtual void SerializationLoad(uint64& InOutOffset, TConstArrayView<uint8> Data) override final
+			{
+				Serialization::Load(InOutOffset, InputSize, Data);
+				Serialization::Load(InOutOffset, OutputSize, Data);
+				Serialization::Load(InOutOffset, WeightOffsets, Data, InputSize);
+				Serialization::Load(InOutOffset, WeightScales, Data, InputSize);
+				Serialization::Load(InOutOffset, Biases, Data, OutputSize);
+				Serialization::Load(InOutOffset, Weights, Data, InputSize * OutputSize);
+			}
+
+			virtual void SerializationSave(uint64& InOutOffset, TArrayView<uint8> Data) const override final
+			{
+				Serialization::Save(InOutOffset, InputSize, Data);
+				Serialization::Save(InOutOffset, OutputSize, Data);
+				Serialization::Save(InOutOffset, WeightOffsets, Data);
+				Serialization::Save(InOutOffset, WeightScales, Data);
+				Serialization::Save(InOutOffset, Biases, Data);
+				Serialization::Save(InOutOffset, Weights, Data);
+			}
+
+			virtual void Evaluate(
+				TArrayView<float> Output,
+				TConstArrayView<float> Input,
+				ILayerInstance* Instance,
+				uint32 BatchSize,
+				uint32 OutputBufferSize,
+				uint32 InputBufferSize) override final
+			{
+				TRACE_CPUPROFILER_EVENT_SCOPE(NNE::RuntimeBasic::Private::FCompressedLinearLayer::Evaluate);
+
+#if NNE_RUNTIME_BASIC_ENABLE_ISPC
+				ispc::NNERuntimeBasicCPUCompressedLinearLayer(
+					Output.GetData(),
+					Input.GetData(),
+					Weights.GetData(),
+					WeightOffsets.GetData(),
+					WeightScales.GetData(),
+					Biases.GetData(),
+					BatchSize,
+					InputSize,
+					OutputSize,
+					OutputBufferSize,
+					InputBufferSize);
+#else
+				CompressedLinearLayer(
+					Output.GetData(),
+					Input.GetData(),
+					Weights.GetData(),
+					WeightOffsets.GetData(),
+					WeightScales.GetData(),
+					Biases.GetData(),
+					BatchSize,
+					InputSize,
+					OutputSize,
+					OutputBufferSize,
+					InputBufferSize);
+#endif
+			}
+
+			uint32 InputSize = 0;
+			uint32 OutputSize = 0;
+			TConstArrayView<float> WeightOffsets;
+			TConstArrayView<float> WeightScales;
+			TConstArrayView<float> Biases;
+			TConstArrayView<uint16> Weights;
+		};
+
+		//--------------------------------------------------------------------------
+
+		static inline void MultiLinearLayer(
+			float* RESTRICT Output,
+			const float* RESTRICT Input,
+			const float* RESTRICT Weights,
+			const float* RESTRICT Biases,
+			const uint32 BatchNum,
+			const uint32 BlockNum,
+			const uint32 RowNum,
+			const uint32 ColNum,
+			const uint32 OutputBufferSize,
+			const uint32 InputBufferSize)
+		{
+			for (uint32 BatchIdx = 0; BatchIdx < BatchNum; BatchIdx++)
+			{
+				for (uint32 BlockIdx = 0; BlockIdx < BlockNum; BlockIdx++)
+				{
+					for (uint32 ColIdx = 0; ColIdx < ColNum; ColIdx++)
+					{
+						Output[BatchIdx * OutputBufferSize + BlockIdx * ColNum + ColIdx] = Biases[BlockIdx * ColNum + ColIdx];
+					}
+				}
+			}
+
+			for (uint32 BatchIdx = 0; BatchIdx < BatchNum; BatchIdx++)
+			{
+				for (uint32 BlockIdx = 0; BlockIdx < BlockNum; BlockIdx++)
+				{
+					for (uint32 RowIdx = 0; RowIdx < RowNum; RowIdx++)
+					{
+						const float Value = Input[BatchIdx * InputBufferSize + BlockIdx * RowNum + RowIdx];
+
+						if (Value != 0.0)
+						{
+							for (uint32 ColIdx = 0; ColIdx < ColNum; ColIdx++)
+							{
+								Output[BatchIdx * OutputBufferSize + BlockIdx * ColNum + ColIdx] += Value * Weights[BlockIdx * RowNum * ColNum + RowIdx * ColNum + ColIdx];
+							}
+						}
+					}
+				}
+			}
+		}
+
+		struct FMultiLinearLayer : public ILayer
+		{
+			virtual ELayerType GetLayerType() const override final { return ELayerType::MultiLinear; }
+			virtual uint32 GetInputSize() const override final { return BlockNum * InputSize; }
+			virtual uint32 GetOutputSize() const override final { return BlockNum * OutputSize; }
+
+			virtual void SerializationSize(uint64& InOutOffset) const override final
+			{
+				Serialization::Size(InOutOffset, BlockNum);
+				Serialization::Size(InOutOffset, InputSize);
+				Serialization::Size(InOutOffset, OutputSize);
+				Serialization::Size(InOutOffset, Biases);
+				Serialization::Size(InOutOffset, Weights);
+			}
+
+			virtual void SerializationLoad(uint64& InOutOffset, TConstArrayView<uint8> Data) override final
+			{
+				Serialization::Load(InOutOffset, BlockNum, Data);
+				Serialization::Load(InOutOffset, InputSize, Data);
+				Serialization::Load(InOutOffset, OutputSize, Data);
+				Serialization::Load(InOutOffset, Biases, Data, BlockNum * OutputSize);
+				Serialization::Load(InOutOffset, Weights, Data, BlockNum * InputSize * OutputSize);
+			}
+
+			virtual void SerializationSave(uint64& InOutOffset, TArrayView<uint8> Data) const override final
+			{
+				Serialization::Save(InOutOffset, BlockNum, Data);
+				Serialization::Save(InOutOffset, InputSize, Data);
+				Serialization::Save(InOutOffset, OutputSize, Data);
+				Serialization::Save(InOutOffset, Biases, Data);
+				Serialization::Save(InOutOffset, Weights, Data);
+			}
+
+			virtual void Evaluate(
+				TArrayView<float> Output,
+				TConstArrayView<float> Input,
+				ILayerInstance* Instance,
+				uint32 BatchSize,
+				uint32 OutputBufferSize,
+				uint32 InputBufferSize) override final
+			{
+				TRACE_CPUPROFILER_EVENT_SCOPE(NNE::RuntimeBasic::Private::FMultiLinearLayer::Evaluate);
+
+				// Here ispc generates slightly less efficient code than the naive C++ implementation
+				// so we want to disable the use of the ispc version unconditionally
+#if 0
+//#if NNE_RUNTIME_BASIC_ENABLE_ISPC
+				ispc::NNERuntimeBasicCPUMultiLinearLayer(
+					Output.GetData(),
+					Input.GetData(),
+					Weights.GetData(),
+					Biases.GetData(),
+					BatchSize,
+					BlockNum,
+					InputSize,
+					OutputSize,
+					OutputBufferSize,
+					InputBufferSize);
+#else
+				MultiLinearLayer(
+					Output.GetData(),
+					Input.GetData(),
+					Weights.GetData(),
+					Biases.GetData(),
+					BatchSize,
+					BlockNum,
+					InputSize,
+					OutputSize,
+					OutputBufferSize,
+					InputBufferSize);
+#endif
+			}
+
+			uint32 InputSize = 0;
+			uint32 OutputSize = 0;
+			uint32 BlockNum = 0;
+			TConstArrayView<float> Biases;
+			TConstArrayView<float> Weights;
+		};
+
+		//--------------------------------------------------------------------------
+
+		static inline void ActivationReLU(
+			float* RESTRICT Output,
+			const float* RESTRICT Input,
+			const uint32 BatchNum,
+			const uint32 Num,
+			const uint32 OutputBufferSize,
+			const uint32 InputBufferSize)
+		{
+			for (uint32 BatchIdx = 0; BatchIdx < BatchNum; BatchIdx++)
+			{
+				for (uint32 Idx = 0; Idx < Num; Idx++)
+				{
+					Output[BatchIdx * OutputBufferSize + Idx] = FMath::Max(Input[BatchIdx * InputBufferSize + Idx], 0.0f);
+				}
+			}
+		}
+
+		struct FReLULayer : public ILayer
+		{
+			virtual ELayerType GetLayerType() const override final { return ELayerType::ReLU; }
+			virtual uint32 GetInputSize() const override final { return InputOutputSize; }
+			virtual uint32 GetOutputSize() const override final { return InputOutputSize; }
+
+			virtual void SerializationSize(uint64& InOutOffset) const override final
+			{
+				Serialization::Size(InOutOffset, InputOutputSize);
+			}
+
+			virtual void SerializationLoad(uint64& InOutOffset, TConstArrayView<uint8> Data) override final
+			{
+				Serialization::Load(InOutOffset, InputOutputSize, Data);
+			}
+
+			virtual void SerializationSave(uint64& InOutOffset, TArrayView<uint8> Data) const override final
+			{
+				Serialization::Save(InOutOffset, InputOutputSize, Data);
+			}
+
+			virtual void Evaluate(
+				TArrayView<float> Output,
+				TConstArrayView<float> Input,
+				ILayerInstance* Instance,
+				uint32 BatchSize,
+				uint32 OutputBufferSize,
+				uint32 InputBufferSize) override final
+			{
+				TRACE_CPUPROFILER_EVENT_SCOPE(NNE::RuntimeBasic::Private::FReLULayer::Evaluate);
+
+#if NNE_RUNTIME_BASIC_ENABLE_ISPC
+				ispc::NNERuntimeBasicCPUActivationReLU(
+					Output.GetData(),
+					Input.GetData(),
+					BatchSize,
+					InputOutputSize,
+					OutputBufferSize,
+					InputBufferSize);
+#else
+				ActivationReLU(
+					Output.GetData(),
+					Input.GetData(),
+					BatchSize,
+					InputOutputSize,
+					OutputBufferSize,
+					InputBufferSize);
+#endif
+			}
+
+			uint32 InputOutputSize = 0;
+		};
+
+		//--------------------------------------------------------------------------
+
+		static inline void ActivationELU(
+			float* RESTRICT Output,
+			const float* RESTRICT Input,
+			const uint32 BatchNum,
+			const uint32 Num,
+			const uint32 OutputBufferSize,
+			const uint32 InputBufferSize)
+		{
+			for (uint32 BatchIdx = 0; BatchIdx < BatchNum; BatchIdx++)
+			{
+				for (uint32 Idx = 0; Idx < Num; Idx++)
+				{
+					const float Value = Input[BatchIdx * InputBufferSize + Idx];
+					Output[BatchIdx * OutputBufferSize + Idx] = Value > 0.0f ? Value : FMath::InvExpApprox(-Value) - 1.0f;
+				}
+			}
+		}
+
+		struct FELULayer : public ILayer
+		{
+			virtual ELayerType GetLayerType() const override final { return ELayerType::ELU; }
+			virtual uint32 GetInputSize() const override final { return InputOutputSize; }
+			virtual uint32 GetOutputSize() const override final { return InputOutputSize; }
+
+			virtual void SerializationSize(uint64& InOutOffset) const override final
+			{
+				Serialization::Size(InOutOffset, InputOutputSize);
+			}
+
+			virtual void SerializationLoad(uint64& InOutOffset, TConstArrayView<uint8> Data) override final
+			{
+				Serialization::Load(InOutOffset, InputOutputSize, Data);
+			}
+
+			virtual void SerializationSave(uint64& InOutOffset, TArrayView<uint8> Data) const override final
+			{
+				Serialization::Save(InOutOffset, InputOutputSize, Data);
+			}
+
+			virtual void Evaluate(
+				TArrayView<float> Output,
+				TConstArrayView<float> Input,
+				ILayerInstance* Instance,
+				uint32 BatchSize,
+				uint32 OutputBufferSize,
+				uint32 InputBufferSize) override final
+			{
+				TRACE_CPUPROFILER_EVENT_SCOPE(NNE::RuntimeBasic::Private::FELULayer::Evaluate);
+
+#if NNE_RUNTIME_BASIC_ENABLE_ISPC
+				ispc::NNERuntimeBasicCPUActivationELU(
+					Output.GetData(),
+					Input.GetData(),
+					BatchSize,
+					InputOutputSize,
+					OutputBufferSize,
+					InputBufferSize);
+#else
+				ActivationELU(
+					Output.GetData(),
+					Input.GetData(),
+					BatchSize,
+					InputOutputSize,
+					OutputBufferSize,
+					InputBufferSize);
+#endif
+			}
+
+			uint32 InputOutputSize = 0;
+		};
+
+		//--------------------------------------------------------------------------
+
+		static inline void ActivationTanH(
+			float* RESTRICT Output,
+			const float* RESTRICT Input,
+			const uint32 BatchNum,
+			const uint32 Num,
+			const uint32 OutputBufferSize,
+			const uint32 InputBufferSize)
+		{
+			for (uint32 BatchIdx = 0; BatchIdx < BatchNum; BatchIdx++)
+			{
+				for (uint32 Idx = 0; Idx < Num; Idx++)
+				{
+					Output[BatchIdx * OutputBufferSize + Idx] = FMath::Tanh(Input[BatchIdx * InputBufferSize + Idx]);
+				}
+			}
+		}
+
+		struct FTanHLayer : public ILayer
+		{
+			virtual ELayerType GetLayerType() const override final { return ELayerType::TanH; }
+			virtual uint32 GetInputSize() const override final { return InputOutputSize; }
+			virtual uint32 GetOutputSize() const override final { return InputOutputSize; }
+
+			virtual void SerializationSize(uint64& InOutOffset) const override final
+			{
+				Serialization::Size(InOutOffset, InputOutputSize);
+			}
+
+			virtual void SerializationLoad(uint64& InOutOffset, TConstArrayView<uint8> Data) override final
+			{
+				Serialization::Load(InOutOffset, InputOutputSize, Data);
+			}
+
+			virtual void SerializationSave(uint64& InOutOffset, TArrayView<uint8> Data) const override final
+			{
+				Serialization::Save(InOutOffset, InputOutputSize, Data);
+			}
+
+			virtual void Evaluate(
+				TArrayView<float> Output,
+				TConstArrayView<float> Input,
+				ILayerInstance* Instance,
+				uint32 BatchSize,
+				uint32 OutputBufferSize,
+				uint32 InputBufferSize) override final
+			{
+				TRACE_CPUPROFILER_EVENT_SCOPE(NNE::RuntimeBasic::Private::FTanHLayer::Evaluate);
+
+#if NNE_RUNTIME_BASIC_ENABLE_ISPC
+				ispc::NNERuntimeBasicCPUActivationTanH(
+					Output.GetData(),
+					Input.GetData(),
+					BatchSize,
+					InputOutputSize,
+					OutputBufferSize,
+					InputBufferSize);
+#else
+				ActivationTanH(
+					Output.GetData(),
+					Input.GetData(),
+					BatchSize,
+					InputOutputSize,
+					OutputBufferSize,
+					InputBufferSize);
+#endif
+			}
+
+			uint32 InputOutputSize = 0;
+		};
+
+		//--------------------------------------------------------------------------
+
+		static inline void ActivationPReLU(
+			float* RESTRICT Output,
+			const float* RESTRICT Input,
+			const float* RESTRICT Alpha,
+			const uint32 BatchNum,
+			const uint32 Num,
+			const uint32 OutputBufferSize,
+			const uint32 InputBufferSize)
+		{
+			for (uint32 BatchIdx = 0; BatchIdx < BatchNum; BatchIdx++)
+			{
+				for (uint32 Idx = 0; Idx < Num; Idx++)
+				{
+					const float Value = Input[BatchIdx * InputBufferSize + Idx];
+					Output[BatchIdx * OutputBufferSize + Idx] = Value > 0.0f ? Value : Alpha[Idx] * Value;
+				}
+			}
+		}
+
+		struct FPReLULayer : public ILayer
+		{
+			virtual ELayerType GetLayerType() const override final { return ELayerType::PReLU; }
+			virtual uint32 GetInputSize() const override final { return InputOutputSize; }
+			virtual uint32 GetOutputSize() const override final { return InputOutputSize; }
+
+			virtual void SerializationSize(uint64& InOutOffset) const override final
+			{
+				Serialization::Size(InOutOffset, InputOutputSize);
+				Serialization::Size(InOutOffset, Alpha);
+			}
+
+			virtual void SerializationLoad(uint64& InOutOffset, TConstArrayView<uint8> Data) override final
+			{
+				Serialization::Load(InOutOffset, InputOutputSize, Data);
+				Serialization::Load(InOutOffset, Alpha, Data, InputOutputSize);
+			}
+
+			virtual void SerializationSave(uint64& InOutOffset, TArrayView<uint8> Data) const override final
+			{
+				Serialization::Save(InOutOffset, InputOutputSize, Data);
+				Serialization::Save(InOutOffset, Alpha, Data);
+			}
+
+			virtual void Evaluate(
+				TArrayView<float> Output,
+				TConstArrayView<float> Input,
+				ILayerInstance* Instance,
+				uint32 BatchSize,
+				uint32 OutputBufferSize,
+				uint32 InputBufferSize) override final
+			{
+				TRACE_CPUPROFILER_EVENT_SCOPE(NNE::RuntimeBasic::Private::FPReLuLayer::Evaluate);
+
+#if NNE_RUNTIME_BASIC_ENABLE_ISPC
+				ispc::NNERuntimeBasicCPUActivationPReLU(
+					Output.GetData(),
+					Input.GetData(),
+					Alpha.GetData(),
+					BatchSize,
+					InputOutputSize,
+					OutputBufferSize,
+					InputBufferSize);
+#else
+				ActivationPReLU(
+					Output.GetData(),
+					Input.GetData(),
+					Alpha.GetData(),
+					BatchSize,
+					InputOutputSize,
+					OutputBufferSize,
+					InputBufferSize);
+#endif
+			}
+
+			uint32 InputOutputSize = 0;
+			TConstArrayView<float> Alpha;
+		};
+
+		//--------------------------------------------------------------------------
+
+		namespace Serialization
+		{
+			static inline void Size(uint64& InOutOffset, const TSharedPtr<ILayer>& InLayer)
+			{
+				Serialization::Size(InOutOffset, (uint32)InLayer->GetLayerType());
+				InLayer->SerializationSize(InOutOffset);
+			}
+
+			static inline void Load(uint64& InOutOffset, TSharedPtr<ILayer>& OutLayer, TConstArrayView<uint8> Data)
+			{
+				uint32 LayerTypeId = (uint32)ELayerType::Invalid;
+				Serialization::Load(InOutOffset, LayerTypeId, Data);
+
+				checkf((ELayerType)LayerTypeId != ELayerType::Invalid, TEXT("Invalid Layer"));
+
+				if (OutLayer == nullptr || OutLayer->GetLayerType() != (ELayerType)LayerTypeId)
+				{
+					switch ((ELayerType)LayerTypeId)
+					{
+					case ELayerType::Sequence: OutLayer = MakeShared<FSequenceLayer>(); break;
+					case ELayerType::Normalize: OutLayer = MakeShared<FNormalizeLayer>(); break;
+					case ELayerType::Denormalize: OutLayer = MakeShared<FDenormalizeLayer>(); break;
+					case ELayerType::Linear: OutLayer = MakeShared<FLinearLayer>(); break;
+					case ELayerType::CompressedLinear: OutLayer = MakeShared<FCompressedLinearLayer>(); break;
+					case ELayerType::MultiLinear: OutLayer = MakeShared<FMultiLinearLayer>(); break;
+					case ELayerType::ReLU: OutLayer = MakeShared<FReLULayer>(); break;
+					case ELayerType::ELU: OutLayer = MakeShared<FELULayer>(); break;
+					case ELayerType::TanH: OutLayer = MakeShared<FTanHLayer>(); break;
+					case ELayerType::PReLU: OutLayer = MakeShared<FPReLULayer>(); break;
+					default: checkf(false, TEXT("Unknown Layer Id %i"), LayerTypeId);
+					}
+				}
+
+				OutLayer->SerializationLoad(InOutOffset, Data);
+			}
+
+			static inline void Save(uint64& InOutOffset, const TSharedPtr<ILayer>& InLayer, TArrayView<uint8> Data)
+			{
+				Serialization::Save(InOutOffset, (uint32)InLayer->GetLayerType(), Data);
+				InLayer->SerializationSave(InOutOffset, Data);
+			}
+		}
+	}
+
+	//--------------------------------------------------------------------------
+
+	FModelInstanceCPU::FModelInstanceCPU(const TSharedPtr<FModelCPU>& InModel)
+		: Model(InModel)
+		, InputTensorDesc(FTensorDesc::Make(TEXT("Input"), FSymbolicTensorShape::Make({ -1, -1 }), ENNETensorDataType::Float))
+		, OutputTensorDesc(FTensorDesc::Make(TEXT("Output"), FSymbolicTensorShape::Make({ -1, -1 }), ENNETensorDataType::Float))
+		, Instance(Model->Layer->MakeInstance())
+	{}
+
+	int FModelInstanceCPU::SetInputTensorShapes(TConstArrayView<FTensorShape> InInputShapes)
+	{
+		TRACE_CPUPROFILER_EVENT_SCOPE(NNE::RuntimeBasic::FModelInstanceCPU::SetInputTensorShapes);
+
+		if (!ensureMsgf(InInputShapes.Num() == 1, TEXT("Basic CPU Inference only supports single input tensor.")))
+		{
+			return -1;
+		}
+
+		const FTensorShape& InputShape = InInputShapes[0];
+
+		if (!ensureMsgf(InputShape.Rank() == 2, TEXT("Basic CPU Inference only supports rank 2 input tensors.")))
+		{
+			return -1;
+		}
+
+		const uint32 InputBatchSize = InputShape.GetData()[0];
+		const uint32 InputInputSize = InputShape.GetData()[1];
+		const uint32 ModelInputSize = Model->Layer->GetInputSize();
+		const uint32 ModelOutputSize = Model->Layer->GetOutputSize();
+
+		if (!ensureMsgf(InputInputSize == ModelInputSize, TEXT("Input tensor shape does not match model input size. Got %i, expected %i."), InputInputSize, ModelInputSize))
+		{
+			return -1;
+		}
+
+		BatchSize = InputBatchSize;
+		InputSize = ModelInputSize;
+		OutputSize = ModelOutputSize;
+
+		InputTensorShape = FTensorShape::Make({ BatchSize, InputSize });
+		OutputTensorShape = FTensorShape::Make({ BatchSize, OutputSize });
+
+		if (Instance)
+		{
+			Instance->SetBatchSize(BatchSize);
+		}
+
+		return 0;
+	}
+
+	int FModelInstanceCPU::RunSync(TConstArrayView<FTensorBindingCPU> InInputBindings, TConstArrayView<FTensorBindingCPU> InOutputBindings)
+	{
+		TRACE_CPUPROFILER_EVENT_SCOPE(NNE::RuntimeBasic::FModelInstanceCPU::RunSync);
+
+		if (!ensureMsgf(BatchSize > 0, TEXT("SetInputTensorShapes must be run before RunSync")))
+		{
+			return -1;
+		}
+
+		if (!ensureMsgf(InInputBindings.Num() == 1, TEXT("Basic CPU Inference only supports single input tensor.")))
+		{
+			return -1;
+		}
+
+		if (!ensureMsgf(InOutputBindings.Num() == 1, TEXT("Basic CPU Inference only supports single output tensor.")))
+		{
+			return -1;
+		}
+
+		if (!ensureMsgf(InInputBindings[0].SizeInBytes == BatchSize * InputSize * sizeof(float), TEXT("Incorrect Input Tensor Size")))
+		{
+			return -1;
+		}
+
+		if (!ensureMsgf(InOutputBindings[0].SizeInBytes == BatchSize * OutputSize * sizeof(float), TEXT("Incorrect Output Tensor Size")))
+		{
+			return -1;
+		}
+
+		Model->Layer->Evaluate(
+			TArrayView<float>((float*)InOutputBindings[0].Data, BatchSize * OutputSize),
+			TConstArrayView<float>((const float*)InInputBindings[0].Data, BatchSize * InputSize),
+			Instance.Get(),
+			BatchSize,
+			OutputSize,
+			InputSize);
+
+		return 0;
+	}
+
+
+	//--------------------------------------------------------------------------
+
+	uint32 FModelCPU::ModelMagicNumber = 0x0BA51C01;
+	uint32 FModelCPU::ModelVersionNumber = 1;
+
+	TSharedPtr<IModelInstanceCPU> FModelCPU::CreateModelInstanceCPU()
+	{
+		return MakeShared<FModelInstanceCPU>(WeakThis.Pin());
+	}
+
+	void FModelCPU::SerializationSize(uint64& InOutOffset) const
+	{
+		checkf(InOutOffset % 64 == 0,
+			TEXT("Model must be aligned to 64 bytes because there must be no padding before magic number."));
+
+		Private::Serialization::Size(InOutOffset, ModelMagicNumber);
+		Private::Serialization::Size(InOutOffset, ModelVersionNumber);
+		Private::Serialization::Size(InOutOffset, Layer);
+	}
+
+	bool FModelCPU::SerializationLoad(uint64& InOutOffset, TConstArrayView<uint8> Data)
+	{
+		TRACE_CPUPROFILER_EVENT_SCOPE(NNE::RuntimeBasic::FModelCPU::SerializationLoad);
+
+		checkf(InOutOffset % 64 == 0, 
+			TEXT("Model must be aligned to 64 bytes because there must be no padding before magic number."));
+
+		uint32 Magic = INDEX_NONE;
+		Private::Serialization::Load(InOutOffset, Magic, Data);
+		if (Magic != ModelMagicNumber)
+		{
+			UE_LOG(LogNNE, Error, TEXT("Invalid Magic Number %i"), Magic);
+			return false;
+		}
+
+		uint32 Version = INDEX_NONE;
+		Private::Serialization::Load(InOutOffset, Version, Data);
+		if (Version != ModelVersionNumber)
+		{
+			UE_LOG(LogNNE, Error, TEXT("Unsupported Version Number %i"), Version);
+			return false;
+		}
+
+		Private::Serialization::Load(InOutOffset, Layer, Data);
+
+		return true;
+	}
+
+	void FModelCPU::SerializationSave(uint64& InOutOffset, TArrayView<uint8> Data) const
+	{
+		TRACE_CPUPROFILER_EVENT_SCOPE(NNE::RuntimeBasic::FModelCPU::SerializationSave);
+
+		checkf(InOutOffset % 64 == 0,
+			TEXT("Model must be aligned to 64 bytes because there must be no padding before magic number."));
+
+		Private::Serialization::Save(InOutOffset, ModelMagicNumber, Data);
+		Private::Serialization::Save(InOutOffset, ModelVersionNumber, Data);
+		Private::Serialization::Save(InOutOffset, Layer, Data);
+	}
+
+	//--------------------------------------------------------------------------
+
+	FSequentialModelBuilder::FSequentialModelBuilder()
+	{
+		Model = MakeShared<FModelCPU>();
+		Model->Layer = MakeShared<Private::FSequenceLayer>();
+	}
+
+	void FSequentialModelBuilder::AddLinear(uint32 InputSize, uint32 OutputSize, TConstArrayView<float> Weights, TConstArrayView<float> Biases)
+	{
+		Private::FSequenceLayer* Sequence = StaticCast<Private::FSequenceLayer*>(Model->Layer.Get());
+
+		if (Sequence->Layers.Num() > 0)
+		{
+			check(Sequence->Layers.Last()->GetOutputSize() == InputSize);
+		}
+
+		check(Weights.Num() == InputSize * OutputSize);
+		check(Biases.Num() == OutputSize);
+
+		const TSharedPtr<Private::FLinearLayer> LinearLayer = MakeShared<Private::FLinearLayer>();
+		LinearLayer->InputSize = InputSize;
+		LinearLayer->OutputSize = OutputSize;
+		LinearLayer->Biases = Biases;
+		LinearLayer->Weights = Weights;
+
+		Sequence->Layers.Emplace(LinearLayer);
+	}
+
+	void FSequentialModelBuilder::AddMultiLinear(uint32 InputSize, uint32 OutputSize, uint32 BlockNum, TConstArrayView<float> Weights, TConstArrayView<float> Biases)
+	{
+		Private::FSequenceLayer* Sequence = StaticCast<Private::FSequenceLayer*>(Model->Layer.Get());
+
+		if (Sequence->Layers.Num() > 0)
+		{
+			check(Sequence->Layers.Last()->GetOutputSize() == InputSize * BlockNum);
+		}
+
+		check(Weights.Num() == InputSize * OutputSize * BlockNum);
+		check(Biases.Num() == OutputSize * BlockNum);
+
+		const TSharedPtr<Private::FMultiLinearLayer> MultiLinearLayer = MakeShared<Private::FMultiLinearLayer>();
+		MultiLinearLayer->InputSize = InputSize;
+		MultiLinearLayer->OutputSize = OutputSize;
+		MultiLinearLayer->BlockNum = BlockNum;
+		MultiLinearLayer->Biases = Biases;
+		MultiLinearLayer->Weights = Weights;
+
+		Sequence->Layers.Emplace(MultiLinearLayer);
+	}
+
+	void FSequentialModelBuilder::AddReLU()
+	{
+		Private::FSequenceLayer* Sequence = StaticCast<Private::FSequenceLayer*>(Model->Layer.Get());
+
+		checkf(Sequence->Layers.Num() > 0, TEXT("Cannot add activation on initial layer because size is unknown"));
+
+		const TSharedPtr<Private::FReLULayer> ActivationLayer = MakeShared<Private::FReLULayer>();
+		ActivationLayer->InputOutputSize = Sequence->Layers.Last()->GetOutputSize();
+		Sequence->Layers.Emplace(ActivationLayer);
+	}
+
+	void FSequentialModelBuilder::AddELU()
+	{
+		Private::FSequenceLayer* Sequence = StaticCast<Private::FSequenceLayer*>(Model->Layer.Get());
+
+		checkf(Sequence->Layers.Num() > 0, TEXT("Cannot add activation on initial layer because size is unknown"));
+
+		const TSharedPtr<Private::FELULayer> ActivationLayer = MakeShared<Private::FELULayer>();
+		ActivationLayer->InputOutputSize = Sequence->Layers.Last()->GetOutputSize();
+		Sequence->Layers.Emplace(ActivationLayer);
+	}
+
+	void FSequentialModelBuilder::AddTanH()
+	{
+		Private::FSequenceLayer* Sequence = StaticCast<Private::FSequenceLayer*>(Model->Layer.Get());
+
+		checkf(Sequence->Layers.Num() > 0, TEXT("Cannot add activation on initial layer because size is unknown"));
+
+		const TSharedPtr<Private::FTanHLayer> ActivationLayer = MakeShared<Private::FTanHLayer>();
+		ActivationLayer->InputOutputSize = Sequence->Layers.Last()->GetOutputSize();
+		Sequence->Layers.Emplace(ActivationLayer);
+	}
+
+	void FSequentialModelBuilder::AddPReLU(TConstArrayView<float> Alpha)
+	{
+		Private::FSequenceLayer* Sequence = StaticCast<Private::FSequenceLayer*>(Model->Layer.Get());
+
+		checkf(Sequence->Layers.Num() > 0, TEXT("Cannot add activation on initial layer because size is unknown"));
+		check(Alpha.Num() == Sequence->Layers.Last()->GetOutputSize());
+
+		const TSharedPtr<Private::FPReLULayer> ActivationLayer = MakeShared<Private::FPReLULayer>();
+		ActivationLayer->InputOutputSize = Sequence->Layers.Last()->GetOutputSize();
+		ActivationLayer->Alpha = Alpha;
+
+		Sequence->Layers.Emplace(ActivationLayer);
+	}
+
+	uint64 FSequentialModelBuilder::GetWriteByteNum() const
+	{
+		uint64 Offset = 0;
+		Model->SerializationSize(Offset);
+		return Offset;
+	}
+
+	void FSequentialModelBuilder::WriteAndReset(TArrayView<uint8> OutBytes)
+	{
+		check((uint64)OutBytes.Num() >= GetWriteByteNum());
+
+		// Zero to ensure any padding due to alignment is always zero
+		FMemory::Memzero(OutBytes.GetData(), OutBytes.Num());
+
+		uint64 Offset = 0;
+		Model->SerializationSave(Offset, OutBytes);
+
+		Private::FSequenceLayer* Sequence = StaticCast<Private::FSequenceLayer*>(Model->Layer.Get());
+		Sequence->Layers.Empty();
+	}
+
+	//--------------------------------------------------------------------------
+
+} // namespace UE::NNE::RuntimeBasic
+
+#undef NNE_RUNTIME_BASIC_ENABLE_ISPC
diff --git a/Engine/Plugins/Experimental/NNERuntimeBasicCpu/Source/NNERuntimeBasicCpu/Private/NNERuntimeBasicCpuModel.h b/Engine/Plugins/Experimental/NNERuntimeBasicCpu/Source/NNERuntimeBasicCpu/Private/NNERuntimeBasicCpuModel.h
new file mode 100644
--- /dev/null
+++ b/Engine/Plugins/Experimental/NNERuntimeBasicCpu/Source/NNERuntimeBasicCpu/Private/NNERuntimeBasicCpuModel.h
@@ -0,0 +1,93 @@
+// Copyright Epic Games, Inc. All Rights Reserved.
+
+#pragma once
+
+#include "NNERuntimeCPU.h"
+#include "NNETensor.h"
+#include "NNETypes.h"
+
+namespace UE::NNE
+{
+	class FSharedModelData;
+}
+
+namespace UE::NNE::RuntimeBasic
+{
+	namespace Private
+	{
+		struct ILayer;
+		struct ILayerInstance;
+	}
+
+	class FModelCPU;
+	class FModelInstanceCPU;
+
+	/**
+	 * Basic implementation of a model instance used for performing inference in this runtime.
+	 */
+	class FModelInstanceCPU : public IModelInstanceCPU
+	{
+	public:
+
+		FModelInstanceCPU(const TSharedPtr<FModelCPU>& InModel);
+
+		//~ Begin UE::NNE::IModelInstanceCPU Interface
+		virtual TConstArrayView<FTensorDesc> GetInputTensorDescs() const override final { return TConstArrayView<FTensorDesc>(&InputTensorDesc, 1); }
+		virtual TConstArrayView<FTensorDesc> GetOutputTensorDescs() const override final { return TConstArrayView<FTensorDesc>(&OutputTensorDesc, 1); }
+		virtual TConstArrayView<FTensorShape> GetInputTensorShapes() const override final { return TConstArrayView<FTensorShape>(&InputTensorShape, 1); }
+		virtual TConstArrayView<FTensorShape> GetOutputTensorShapes() const override final { return TConstArrayView<FTensorShape>(&OutputTensorShape, 1); }
+		virtual int SetInputTensorShapes(TConstArrayView<FTensorShape> InInputShapes) override final;
+		virtual int RunSync(TConstArrayView<FTensorBindingCPU> InInputBindings, TConstArrayView<FTensorBindingCPU> InOutputBindings) override final;
+		//~ End UE::NNE::IModelInstanceCPU Interface
+
+		TSharedPtr<FModelCPU> Model;
+		FTensorDesc InputTensorDesc;
+		FTensorDesc OutputTensorDesc;
+		FTensorShape InputTensorShape;
+		FTensorShape OutputTensorShape;
+		TSharedPtr<Private::ILayerInstance> Instance;
+
+		uint32 BatchSize = 0;
+		uint32 InputSize = 0;
+		uint32 OutputSize = 0;
+	};
+
+	/**
+	 * Basic implementation of a model storing the network weights in this runtime.
+	 */
+	class FModelCPU : public IModelCPU
+	{
+	public:
+
+		/** Magic number for the serialization format of this model data */
+		static uint32 ModelMagicNumber;
+
+		/** Version number for the serialization format of this model data */
+		static uint32 ModelVersionNumber;
+
+		//~ Begin UE::NNE::IModelCPU Interface
+		virtual TSharedPtr<IModelInstanceCPU> CreateModelInstanceCPU() override final;
+		//~ End UE::NNE::IModelCPU Interface
+
+		/** Computes the number of bytes required to save or load this model */
+		void SerializationSize(uint64& InOutOffset) const;
+
+		/** Loads the model from the array of bytes starting at the given offset */
+		bool SerializationLoad(uint64& InOutOffset, TConstArrayView<uint8> Data);
+
+		/** Saves the model to the array of bytes starting at the given offset */
+		void SerializationSave(uint64& InOutOffset, TArrayView<uint8> Data) const;
+
+		/** Private implementation of the model */
+		TSharedPtr<Private::ILayer> Layer;
+
+		/** Shared pointer to the actual model data */
+		TSharedPtr<FSharedModelData> ModelData;
+
+		/**
+		 * We need to store a TWeakPtr to ourselves so that when we construct the instance we can pass it a TSharedPtr 
+		 * to this object, which ensures the Model does not get deleted while there are still instances alive.
+		 */
+		TWeakPtr<FModelCPU> WeakThis;
+	};
+}
\ No newline at end of file
diff --git a/Engine/Plugins/Experimental/NNERuntimeBasicCpu/Source/NNERuntimeBasicCpu/Private/NNERuntimeBasicCpuModule.cpp b/Engine/Plugins/Experimental/NNERuntimeBasicCpu/Source/NNERuntimeBasicCpu/Private/NNERuntimeBasicCpuModule.cpp
new file mode 100644
--- /dev/null
+++ b/Engine/Plugins/Experimental/NNERuntimeBasicCpu/Source/NNERuntimeBasicCpu/Private/NNERuntimeBasicCpuModule.cpp
@@ -0,0 +1,33 @@
+// Copyright Epic Games, Inc. All Rights Reserved.
+
+#include "NNERuntimeBasicCpuModule.h"
+#include "NNE.h"
+#include "NNERuntimeBasicCpu.h"
+#include "UObject/WeakInterfacePtr.h"
+
+void FNNERuntimeBasicCpuModule::StartupModule()
+{
+	// NNE runtime startup
+	NNERuntimeBasicCpu = NewObject<UNNERuntimeBasicCpuImpl>();
+	if (NNERuntimeBasicCpu.IsValid())
+	{
+		TWeakInterfacePtr<INNERuntime> RuntimeCPUInterface(NNERuntimeBasicCpu.Get());
+
+		NNERuntimeBasicCpu->AddToRoot();
+		UE::NNE::RegisterRuntime(RuntimeCPUInterface);
+	}
+}
+
+void FNNERuntimeBasicCpuModule::ShutdownModule()
+{
+	// NNE runtime shutdown
+	if (NNERuntimeBasicCpu.IsValid())
+	{
+		TWeakInterfacePtr<INNERuntime> RuntimeCPUInterface(NNERuntimeBasicCpu.Get());
+
+		UE::NNE::UnregisterRuntime(RuntimeCPUInterface);
+		NNERuntimeBasicCpu.Reset();
+	}
+}
+
+IMPLEMENT_MODULE(FNNERuntimeBasicCpuModule, NNERuntimeBasicCpu);
diff --git a/Engine/Plugins/Experimental/NNERuntimeBasicCpu/Source/NNERuntimeBasicCpu/Private/NNERuntimeBasicCpuModule.h b/Engine/Plugins/Experimental/NNERuntimeBasicCpu/Source/NNERuntimeBasicCpu/Private/NNERuntimeBasicCpuModule.h
new file mode 100644
--- /dev/null
+++ b/Engine/Plugins/Experimental/NNERuntimeBasicCpu/Source/NNERuntimeBasicCpu/Private/NNERuntimeBasicCpuModule.h
@@ -0,0 +1,20 @@
+// Copyright Epic Games, Inc. All Rights Reserved.
+
+#pragma once
+
+#include "CoreMinimal.h"
+#include "Modules/ModuleManager.h"
+#include "UObject/WeakObjectPtr.h"
+#include "UObject/WeakObjectPtrTemplates.h"
+
+class UNNERuntimeBasicCpuImpl;
+
+class FNNERuntimeBasicCpuModule : public IModuleInterface
+{
+public:
+	TWeakObjectPtr<UNNERuntimeBasicCpuImpl> NNERuntimeBasicCpu{ nullptr };
+
+	// Begin IModuleInterface
+	virtual void StartupModule() override;
+	virtual void ShutdownModule() override;
+};
\ No newline at end of file
diff --git a/Engine/Plugins/Experimental/NNERuntimeBasicCpu/Source/NNERuntimeBasicCpu/Public/NNERuntimeBasicCpu.h b/Engine/Plugins/Experimental/NNERuntimeBasicCpu/Source/NNERuntimeBasicCpu/Public/NNERuntimeBasicCpu.h
new file mode 100644
--- /dev/null
+++ b/Engine/Plugins/Experimental/NNERuntimeBasicCpu/Source/NNERuntimeBasicCpu/Public/NNERuntimeBasicCpu.h
@@ -0,0 +1,116 @@
+// Copyright Epic Games, Inc. All Rights Reserved.
+
+#pragma once
+
+#include "NNERuntime.h"
+#include "NNERuntimeCPU.h"
+#include "UObject/Object.h"
+
+#include "NNERuntimeBasicCpu.generated.h"
+
+/**
+ * This plugin is a basic, performant, cross-platform CPU runtime for NNE that supports simple models such as MLPs.
+ * 
+ * To use this runtime, the custom ".ubnne" file format is used, which can be exported from python using the functions
+ * in the provided "nne_runtime_basic_cpu.py" found in the "Content" folder of this plugin. The idea behind this plugin
+ * is not to be a general purpose runtime, but rather to provide performant cross-platform implementations for simple
+ * CPU models such as MLPs with minimal overhead and memory usage.
+ */
+UCLASS()
+class NNERUNTIMEBASICCPU_API UNNERuntimeBasicCpuImpl : public UObject, public INNERuntime, public INNERuntimeCPU
+{
+	GENERATED_BODY()
+
+public:
+
+	virtual FString GetRuntimeName() const override { return TEXT("NNERuntimeBasicCpu"); };
+
+	virtual bool CanCreateModelData(FString FileType, TConstArrayView<uint8> FileData, FGuid FileId, const ITargetPlatform* TargetPlatform) const override;
+	virtual TSharedPtr<UE::NNE::FSharedModelData> CreateModelData(FString FileType, TConstArrayView<uint8> FileData, FGuid FileId, const ITargetPlatform* TargetPlatform) override;
+	virtual FString GetModelDataIdentifier(FString FileType, TConstArrayView<uint8> FileData, FGuid FileId, const ITargetPlatform* TargetPlatform) override;
+
+	virtual bool CanCreateModelCPU(TObjectPtr<UNNEModelData> ModelData) const override;
+	virtual TSharedPtr<UE::NNE::IModelCPU> CreateModelCPU(TObjectPtr<UNNEModelData> ModelData) override;
+
+private:
+
+	static const uint32 Alignment;
+};
+
+namespace UE::NNE::RuntimeBasic
+{
+	class FModelCPU;
+
+	/**
+	 * This class can be used to construct file data for basic sequential models at runtime. It should not live longer
+	 * than the views of data passed in by the various functions for adding layers.
+	 */
+	class NNERUNTIMEBASICCPU_API FSequentialModelBuilder
+	{
+
+	public:
+
+		FSequentialModelBuilder();
+
+		/**
+		 * Adds a new linear layer to the model.
+		 *
+		 * @param InputSize		Input Vector Size (Number of Rows)
+		 * @param OutputSize	Output Vector Size (Number of Columns)
+		 * @param Weights		Linear layer weights.
+		 * @param Biases		Linear layer biases.
+		 */
+		void AddLinear(
+			uint32 InputSize,
+			uint32 OutputSize,
+			TConstArrayView<float> Weights,
+			TConstArrayView<float> Biases);
+
+		/**
+		 * Adds a new multi linear layer to the model.
+		 *
+		 * @param InputSize		Input Vector Size (Number of Rows)
+		 * @param OutputSize	Output Vector Size (Number of Columns)
+		 * @param BlockNum		Number of blocks (Number of Matrices)
+		 * @param Weights		Multi-Linear layer weights.
+		 * @param Biases		Multi-Linear layer biases.
+		 */
+		void AddMultiLinear(
+			uint32 InputSize,
+			uint32 OutputSize,
+			uint32 BlockNum,
+			TConstArrayView<float> Weights,
+			TConstArrayView<float> Biases);
+
+		/** Add a ReLU Activation Layer */
+		void AddReLU();
+
+		/** Add a ELU Activation Layer */
+		void AddELU();
+
+		/** Add a TanH Activation Layer */
+		void AddTanH();
+
+		/**
+		 * Add a PReLU Activation Layer
+		 *
+		 * @param Alpha				PReLU alpha parameter for each neuron.
+		 */
+		void AddPReLU(TConstArrayView<float> Alpha);
+
+		/**
+		 * Get the number of bytes this builder currently wants to write.
+		 */
+		uint64 GetWriteByteNum() const;
+
+		/**
+		 * Write the Model to FileData and reset this builder, freeing all the memory used. Use `GetWriteByteNum` to
+		 * get the number of bytes this will write so that `FileData` can be allocated to the right size.
+		 */
+		void WriteAndReset(TArrayView<uint8> FileData);
+
+	private:
+
+		TSharedPtr<FModelCPU> Model;
+	};
+}
\ No newline at end of file
