LearningAgents: Large restructuring to allow for structured actions and observations.

diff --git a/Engine/Plugins/Experimental/NNERuntimeBasicCpu/Source/NNERuntimeBasicCpu/Private/NNERuntimeBasicCpuModel.cpp b/Engine/Plugins/Experimental/NNERuntimeBasicCpu/Source/NNERuntimeBasicCpu/Private/NNERuntimeBasicCpuModel.cpp
--- a/Engine/Plugins/Experimental/NNERuntimeBasicCpu/Source/NNERuntimeBasicCpu/Private/NNERuntimeBasicCpuModel.cpp
+++ b/Engine/Plugins/Experimental/NNERuntimeBasicCpu/Source/NNERuntimeBasicCpu/Private/NNERuntimeBasicCpuModel.cpp
@@ -4,10 +4,23 @@
 
 #include "NNE.h"
 #include "NNERuntimeBasicCpu.h"
+#include "NNERuntimeBasicCpuBuilder.h"
 
 #define NNE_RUNTIME_BASIC_ENABLE_ISPC INTEL_ISPC
 //#define NNE_RUNTIME_BASIC_ENABLE_ISPC 0
 
+//#define NNE_RUNTIME_BASIC_ENABLE_NAN_CHECK (!UE_BUILD_SHIPPING)
+#define NNE_RUNTIME_BASIC_ENABLE_NAN_CHECK 0
+
+#define NNE_RUNTIME_BASIC_ENABLE_PROFILE (!UE_BUILD_SHIPPING)
+//#define NNE_RUNTIME_BASIC_ENABLE_PROFILE 0
+
+#if NNE_RUNTIME_BASIC_ENABLE_PROFILE
+#define NNE_RUNTIME_BASIC_TRACE_SCOPE(...) TRACE_CPUPROFILER_EVENT_SCOPE(__VA_ARGS__)
+#else
+#define NNE_RUNTIME_BASIC_TRACE_SCOPE(...)
+#endif
+
 #if NNE_RUNTIME_BASIC_ENABLE_ISPC
 #include "NNERuntimeBasicCpu.ispc.generated.h"
 #endif
@@ -16,6 +29,10 @@ namespace UE::NNE::RuntimeBasic
 {
 	namespace Private
 	{
+		//--------------------------------------------------------------------------
+		// Serialization Helpers
+		//--------------------------------------------------------------------------
+
 		namespace Serialization
 		{
 			//--------------------------------------------------------------------------
@@ -51,73 +68,98 @@ namespace UE::NNE::RuntimeBasic
 				InOutOffset += In.Num() * sizeof(uint16);
 			}
 
+			static inline void Size(uint64& InOutOffset, const TConstArrayView<uint32> In)
+			{
+				Align(InOutOffset, 64);
+				InOutOffset += In.Num() * sizeof(uint32);
+			}
+
 			static inline void Size(uint64& InOutOffset, const TSharedPtr<ILayer>& InLayer);
+			static inline void Size(uint64& InOutOffset, const TConstArrayView<TSharedPtr<ILayer>> InLayers);
 
 			//--------------------------------------------------------------------------
 
 			static inline void Load(uint64& InOutOffset, uint32& Out, TConstArrayView<uint8> Data)
 			{
 				Align(InOutOffset, sizeof(uint32));
-				Out = *((uint32*)&Data[InOutOffset]);
+				Out = *((uint32*)(Data.GetData() + InOutOffset));
 				InOutOffset += sizeof(uint32);
 			}
 
 			static inline void Load(uint64& InOutOffset, float& Out, TConstArrayView<uint8> Data)
 			{
 				Align(InOutOffset, sizeof(float));
-				Out = *((float*)&Data[InOutOffset]);
+				Out = *((float*)(Data.GetData() + InOutOffset));
 				InOutOffset += sizeof(float);
 			}
 
 			static inline void Load(uint64& InOutOffset, TConstArrayView<float>& Out, TConstArrayView<uint8> Data, int32 Size)
 			{
 				Align(InOutOffset, 64);
-				Out = MakeArrayView<float>((float*)&Data[InOutOffset], Size);
+				Out = MakeArrayView<const float>((const float*)(Data.GetData() + InOutOffset), Size);
 				InOutOffset += Size * sizeof(float);
 			}
 
 			static inline void Load(uint64& InOutOffset, TConstArrayView<uint16>& Out, TConstArrayView<uint8> Data, int32 Size)
 			{
 				Align(InOutOffset, 64);
-				Out = MakeArrayView<uint16>((uint16*)&Data[InOutOffset], Size);
+				Out = MakeArrayView<const uint16>((const uint16*)(Data.GetData() + InOutOffset), Size);
 				InOutOffset += Size * sizeof(uint16);
 			}
 
+			static inline void Load(uint64& InOutOffset, TConstArrayView<uint32>& Out, TConstArrayView<uint8> Data, int32 Size)
+			{
+				Align(InOutOffset, 64);
+				Out = MakeArrayView<const uint32>((const uint32*)(Data.GetData() + InOutOffset), Size);
+				InOutOffset += Size * sizeof(uint32);
+			}
+
 			static inline void Load(uint64& InOutOffset, TSharedPtr<ILayer>& OutLayer, TConstArrayView<uint8> Data);
+			static inline void Load(uint64& InOutOffset, TArrayView<TSharedPtr<ILayer>> OutLayers, TConstArrayView<uint8> Data);
 
 			//--------------------------------------------------------------------------
 
 			static inline void Save(uint64& InOutOffset, const uint32 In, TArrayView<uint8> Data)
 			{
 				Align(InOutOffset, sizeof(uint32));
-				*((uint32*)&Data[InOutOffset]) = In;
+				*((uint32*)(Data.GetData() + InOutOffset)) = In;
 				InOutOffset += sizeof(uint32);
 			}
 
 			static inline void Save(uint64& InOutOffset, const float In, TArrayView<uint8> Data)
 			{
 				Align(InOutOffset, sizeof(float));
-				*((float*)&Data[InOutOffset]) = In;
+				*((float*)(Data.GetData() + InOutOffset)) = In;
 				InOutOffset += sizeof(float);
 			}
 
 			static inline void Save(uint64& InOutOffset, TConstArrayView<float> In, TArrayView<uint8> Data)
 			{
 				Align(InOutOffset, 64);
-				FMemory::Memcpy(&Data[InOutOffset], In.GetData(), In.Num() * sizeof(float));
+				FMemory::Memcpy(Data.GetData() + InOutOffset, In.GetData(), In.Num() * sizeof(float));
 				InOutOffset += In.Num() * sizeof(float);
 			}
 
 			static inline void Save(uint64& InOutOffset, TConstArrayView<uint16> In, TArrayView<uint8> Data)
 			{
 				Align(InOutOffset, 64);
-				FMemory::Memcpy(&Data[InOutOffset], In.GetData(), In.Num() * sizeof(uint16));
+				FMemory::Memcpy(Data.GetData() + InOutOffset, In.GetData(), In.Num() * sizeof(uint16));
 				InOutOffset += In.Num() * sizeof(uint16);
 			}
 
+			static inline void Save(uint64& InOutOffset, TConstArrayView<uint32> In, TArrayView<uint8> Data)
+			{
+				Align(InOutOffset, 64);
+				FMemory::Memcpy(Data.GetData() + InOutOffset, In.GetData(), In.Num() * sizeof(uint32));
+				InOutOffset += In.Num() * sizeof(uint32);
+			}
+
 			static inline void Save(uint64& InOutOffset, const TSharedPtr<ILayer>& InLayer, TArrayView<uint8> Data);
+			static inline void Save(uint64& InOutOffset, const TConstArrayView<TSharedPtr<ILayer>> InLayers, TArrayView<uint8> Data);
 		}
 
+		//--------------------------------------------------------------------------
+		// Basic Mathematical Functions
 		//--------------------------------------------------------------------------
 
 		static inline float Sigmoid(const float X)
@@ -126,8 +168,32 @@ namespace UE::NNE::RuntimeBasic
 		}
 
 		//--------------------------------------------------------------------------
+		// Operators
+		//--------------------------------------------------------------------------
+
+		static inline void OperatorNanCheck(
+			const float* RESTRICT InputOutput, 
+			const uint32 BatchSize,
+			const uint32 InputOutputSize,
+			const uint32 InputOutputStride)
+		{
+#if NNE_RUNTIME_BASIC_ENABLE_NAN_CHECK
+
+			NNE_RUNTIME_BASIC_TRACE_SCOPE(NNE::RuntimeBasic::Private::OperatorNanCheck);
+
+			for (uint32 BatchIdx = 0; BatchIdx < BatchSize; BatchIdx++)
+			{
+				for (uint32 Idx = 0; Idx < InputOutputSize; Idx++)
+				{
+					const float Value = InputOutput[BatchIdx * InputOutputStride + Idx];
+					checkf(FMath::IsFinite(Value) && Value != MAX_flt && Value != -MAX_flt,
+						TEXT("Invalid value %f found in Batch %i, Value %i"), Value, BatchIdx, Idx);
+				}
+			}
+#endif
+		}
 
-		static inline void IdentityLayer(
+		static inline void OperatorCopy(
 			float* RESTRICT Output,
 			const float* RESTRICT Input,
 			const uint32 BatchSize,
@@ -135,8 +201,10 @@ namespace UE::NNE::RuntimeBasic
 			const uint32 OutputStride,
 			const uint32 InputStride)
 		{
+			NNE_RUNTIME_BASIC_TRACE_SCOPE(NNE::RuntimeBasic::Private::OperatorCopy);
+
 #if NNE_RUNTIME_BASIC_ENABLE_ISPC
-			ispc::NNERuntimeBasicCPUIdentityLayer(
+			ispc::NNERuntimeBasicCPUOperatorCopy(
 				Output,
 				Input,
 				BatchSize,
@@ -154,7 +222,7 @@ namespace UE::NNE::RuntimeBasic
 #endif
 		}
 
-		static inline void NormalizeLayer(
+		static inline void OperatorNormalize(
 			float* RESTRICT Output,
 			const float* RESTRICT Input,
 			const float* RESTRICT Mean,
@@ -164,8 +232,10 @@ namespace UE::NNE::RuntimeBasic
 			const uint32 OutputStride,
 			const uint32 InputStride)
 		{
+			NNE_RUNTIME_BASIC_TRACE_SCOPE(NNE::RuntimeBasic::Private::OperatorNormalize);
+
 #if NNE_RUNTIME_BASIC_ENABLE_ISPC
-			ispc::NNERuntimeBasicCPUNormalizeLayer(
+			ispc::NNERuntimeBasicCPUOperatorNormalize(
 				Output,
 				Input,
 				Mean,
@@ -185,7 +255,7 @@ namespace UE::NNE::RuntimeBasic
 #endif
 		}
 
-		static inline void DenormalizeLayer(
+		static inline void OperatorDenormalize(
 			float* RESTRICT Output,
 			const float* RESTRICT Input,
 			const float* RESTRICT Mean,
@@ -195,8 +265,10 @@ namespace UE::NNE::RuntimeBasic
 			const uint32 OutputStride,
 			const uint32 InputStride)
 		{
+			NNE_RUNTIME_BASIC_TRACE_SCOPE(NNE::RuntimeBasic::Private::OperatorDenormalize);
+
 #if NNE_RUNTIME_BASIC_ENABLE_ISPC
-			ispc::NNERuntimeBasicCPUDenormalizeLayer(
+			ispc::NNERuntimeBasicCPUOperatorDenormalize(
 				Output,
 				Input,
 				Mean,
@@ -216,26 +288,28 @@ namespace UE::NNE::RuntimeBasic
 #endif
 		}
 
-		static inline void LinearLayer(
+		static inline void OperatorLinear(
 			float* RESTRICT Output,
 			const float* RESTRICT Input,
 			const float* RESTRICT Weights,
 			const float* RESTRICT Biases,
 			const uint32 BatchSize,
-			const uint32 InputSize,
 			const uint32 OutputSize,
+			const uint32 InputSize,
 			const uint32 OutputStride,
 			const uint32 InputStride)
 		{
+			NNE_RUNTIME_BASIC_TRACE_SCOPE(NNE::RuntimeBasic::Private::OperatorLinear);
+
 #if NNE_RUNTIME_BASIC_ENABLE_ISPC
-			ispc::NNERuntimeBasicCPULinearLayer(
+			ispc::NNERuntimeBasicCPUOperatorLinear(
 				Output,
 				Input,
 				Weights,
 				Biases,
 				BatchSize,
-				InputSize,
 				OutputSize,
+				InputSize,
 				OutputStride,
 				InputStride);
 #else
@@ -265,7 +339,7 @@ namespace UE::NNE::RuntimeBasic
 #endif
 		}
 
-		static inline void CompressedLinearLayer(
+		static inline void OperatorCompressedLinear(
 			float* RESTRICT Output,
 			const float* RESTRICT Input,
 			const uint16* RESTRICT Weights,
@@ -273,13 +347,15 @@ namespace UE::NNE::RuntimeBasic
 			const float* RESTRICT WeightScales,
 			const float* RESTRICT Biases,
 			const uint32 BatchSize,
-			const uint32 InputSize,
 			const uint32 OutputSize,
+			const uint32 InputSize,
 			const uint32 OutputStride,
 			const uint32 InputStride)
 		{
+			NNE_RUNTIME_BASIC_TRACE_SCOPE(NNE::RuntimeBasic::Private::OperatorCompressedLinear);
+
 #if NNE_RUNTIME_BASIC_ENABLE_ISPC
-			ispc::NNERuntimeBasicCPUCompressedLinearLayer(
+			ispc::NNERuntimeBasicCPUOperatorCompressedLinear(
 				Output,
 				Input,
 				Weights,
@@ -287,8 +363,8 @@ namespace UE::NNE::RuntimeBasic
 				WeightScales,
 				Biases,
 				BatchSize,
-				InputSize,
 				OutputSize,
+				InputSize,
 				OutputStride,
 				InputStride);
 #else
@@ -321,18 +397,20 @@ namespace UE::NNE::RuntimeBasic
 #endif
 		}
 
-		static inline void MultiLinearLayer(
+		static inline void OperatorMultiLinear(
 			float* RESTRICT Output,
 			const float* RESTRICT Input,
 			const float* RESTRICT Weights,
 			const float* RESTRICT Biases,
 			const uint32 BatchSize,
 			const uint32 BlockNum,
-			const uint32 InputSize,
 			const uint32 OutputSize,
+			const uint32 InputSize,
 			const uint32 OutputStride,
 			const uint32 InputStride)
 		{
+			NNE_RUNTIME_BASIC_TRACE_SCOPE(NNE::RuntimeBasic::Private::OperatorMultiLinear);
+
 			// For this function ispc generates slightly less efficient code than the naive C++ implementation so we
 			// don't bother calling out to the ispc version even if it is available
 
@@ -367,7 +445,7 @@ namespace UE::NNE::RuntimeBasic
 			}
 		}
 
-		static inline void ActivationReLU(
+		static inline void OperatorReLU(
 			float* RESTRICT Output,
 			const float* RESTRICT Input,
 			const uint32 BatchSize,
@@ -375,8 +453,10 @@ namespace UE::NNE::RuntimeBasic
 			const uint32 OutputStride,
 			const uint32 InputStride)
 		{
+			NNE_RUNTIME_BASIC_TRACE_SCOPE(NNE::RuntimeBasic::Private::OperatorReLU);
+
 #if NNE_RUNTIME_BASIC_ENABLE_ISPC
-			ispc::NNERuntimeBasicCPUActivationReLU(
+			ispc::NNERuntimeBasicCPUOperatorReLU(
 				Output,
 				Input,
 				BatchSize,
@@ -394,7 +474,7 @@ namespace UE::NNE::RuntimeBasic
 #endif
 		}
 
-		static inline void ActivationELU(
+		static inline void OperatorELU(
 			float* RESTRICT Output,
 			const float* RESTRICT Input,
 			const uint32 BatchSize,
@@ -402,8 +482,10 @@ namespace UE::NNE::RuntimeBasic
 			const uint32 OutputStride,
 			const uint32 InputStride)
 		{
+			NNE_RUNTIME_BASIC_TRACE_SCOPE(NNE::RuntimeBasic::Private::OperatorELU);
+
 #if NNE_RUNTIME_BASIC_ENABLE_ISPC
-			ispc::NNERuntimeBasicCPUActivationELU(
+			ispc::NNERuntimeBasicCPUOperatorELU(
 				Output,
 				Input,
 				BatchSize,
@@ -422,7 +504,7 @@ namespace UE::NNE::RuntimeBasic
 #endif
 		}
 
-		static inline void ActivationTanH(
+		static inline void OperatorTanH(
 			float* RESTRICT Output,
 			const float* RESTRICT Input,
 			const uint32 BatchSize,
@@ -430,8 +512,10 @@ namespace UE::NNE::RuntimeBasic
 			const uint32 OutputStride,
 			const uint32 InputStride)
 		{
+			NNE_RUNTIME_BASIC_TRACE_SCOPE(NNE::RuntimeBasic::Private::OperatorTanH);
+
 #if NNE_RUNTIME_BASIC_ENABLE_ISPC
-			ispc::NNERuntimeBasicCPUActivationTanH(
+			ispc::NNERuntimeBasicCPUOperatorTanH(
 				Output,
 				Input,
 				BatchSize,
@@ -449,7 +533,7 @@ namespace UE::NNE::RuntimeBasic
 #endif
 		}
 
-		static inline void ActivationPReLU(
+		static inline void OperatorPReLU(
 			float* RESTRICT Output,
 			const float* RESTRICT Input,
 			const float* RESTRICT Alpha,
@@ -458,8 +542,10 @@ namespace UE::NNE::RuntimeBasic
 			const uint32 OutputStride,
 			const uint32 InputStride)
 		{
+			NNE_RUNTIME_BASIC_TRACE_SCOPE(NNE::RuntimeBasic::Private::OperatorPReLU);
+
 #if NNE_RUNTIME_BASIC_ENABLE_ISPC
-			ispc::NNERuntimeBasicCPUActivationPReLU(
+			ispc::NNERuntimeBasicCPUOperatorPReLU(
 				Output,
 				Input,
 				Alpha,
@@ -479,7 +565,7 @@ namespace UE::NNE::RuntimeBasic
 #endif
 		}
 
-		static inline void MemoryCellUpdateMemory(
+		static inline void OperatorMemoryCellUpdateMemory(
 			float* RESTRICT Output,
 			const float* RESTRICT RememberGate,
 			const float* RESTRICT Memory,
@@ -491,6 +577,21 @@ namespace UE::NNE::RuntimeBasic
 			const uint32 MemoryStride,
 			const uint32 UpdateStride)
 		{
+			NNE_RUNTIME_BASIC_TRACE_SCOPE(NNE::RuntimeBasic::Private::OperatorMemoryCellUpdateMemory);
+
+#if NNE_RUNTIME_BASIC_ENABLE_ISPC
+			ispc::NNERuntimeBasicCPUOperatorMemoryCellUpdateMemory(
+				Output,
+				RememberGate,
+				Memory,
+				Update,
+				BatchSize,
+				MemorySize,
+				OutputStride,
+				RememberGateStride,
+				MemoryStride,
+				UpdateStride);
+#else
 			for (uint32 BatchIdx = 0; BatchIdx < BatchSize; BatchIdx++)
 			{
 				for (uint32 Idx = 0; Idx < MemorySize; Idx++)
@@ -502,9 +603,10 @@ namespace UE::NNE::RuntimeBasic
 					Output[BatchIdx * OutputStride + Idx] = (1.0f - Gate) * Prev + Gate * Targ;
 				}
 			}
+#endif
 		}
 
-		static inline void MemoryCellUpdateOutput(
+		static inline void OperatorMemoryCellUpdateOutput(
 			float* RESTRICT Output,
 			const float* RESTRICT PassthroughGate,
 			const float* RESTRICT MemoryUpdate,
@@ -516,6 +618,21 @@ namespace UE::NNE::RuntimeBasic
 			const uint32 MemoryUpdateStride,
 			const uint32 InputUpdateStride)
 		{
+			NNE_RUNTIME_BASIC_TRACE_SCOPE(NNE::RuntimeBasic::Private::OperatorMemoryCellUpdateOutput);
+
+#if NNE_RUNTIME_BASIC_ENABLE_ISPC
+			ispc::NNERuntimeBasicCPUOperatorMemoryCellUpdateOutput(
+				Output,
+				PassthroughGate,
+				MemoryUpdate,
+				InputUpdate,
+				BatchSize,
+				OutputSize,
+				OutputStride,
+				PassthroughGateStride,
+				MemoryUpdateStride,
+				InputUpdateStride);
+#else
 			for (uint32 BatchIdx = 0; BatchIdx < BatchSize; BatchIdx++)
 			{
 				for (uint32 Idx = 0; Idx < OutputSize; Idx++)
@@ -527,190 +644,677 @@ namespace UE::NNE::RuntimeBasic
 					Output[BatchIdx * OutputStride + Idx] = (1.0f - Gate) * MemTarg + Gate * InTarg;
 				}
 			}
+#endif
 		}
 
-		//--------------------------------------------------------------------------
-
-		/** Layer Type Id - this should match what is given in nne_runtime_basic_cpu.py  */
-		enum class ELayerType : uint32
+		static inline void OperatorAggregateGatherElements(
+			float* RESTRICT OutputBuffer,
+			const float* RESTRICT InputBuffer,
+			const uint32* RESTRICT ElementNums,
+			const uint32* RESTRICT ElementOffsets,
+			const uint32 BatchSize,
+			const uint32 ElementSize,
+			const uint32 InputStride)
 		{
-			Invalid = 0,
-			Sequence = 1,
-			Normalize = 2,
-			Denormalize = 3,
-			Linear = 4,
-			CompressedLinear = 5,
-			MultiLinear = 6,
-			ReLU = 7,
-			ELU = 8,
-			TanH = 9,
-			PReLU = 10,
-			MemoryCell = 11,
-			MemoryBackbone = 12,
-		};
+			NNE_RUNTIME_BASIC_TRACE_SCOPE(NNE::RuntimeBasic::Private::OperatorAggregateGatherElements);
 
-		//--------------------------------------------------------------------------
+#if NNE_RUNTIME_BASIC_ENABLE_ISPC
+			ispc::NNERuntimeBasicCPUOperatorAggregateGatherElements(
+				OutputBuffer,
+				InputBuffer,
+				ElementNums,
+				ElementOffsets,
+				BatchSize,
+				ElementSize,
+				InputStride);
+#else
+			for (uint32 BatchIdx = 0; BatchIdx < BatchSize; BatchIdx++)
+			{
+				const uint32 ElementNum = ElementNums[BatchIdx];
+				const uint32 ElementOffset = ElementOffsets[BatchIdx];
 
-		/**
-		 * Interface for a Layer Instance - the data required for performing inference for a layer.
-		 */
-		struct ILayerInstance
+				for (uint32 ElementIdx = 0; ElementIdx < ElementNum; ElementIdx++)
+				{
+					for (uint32 Idx = 0; Idx < ElementSize; Idx++)
+					{
+						OutputBuffer[(ElementOffset + ElementIdx) * ElementSize + Idx] = InputBuffer[BatchIdx * InputStride + ElementIdx * ElementSize + Idx];
+					}
+				}
+			}
+#endif
+		}
+
+		static inline void OperatorAggregateInsertOneHot(
+			float* RESTRICT QueryBuffer,
+			const uint32 Index,
+			const uint32 BatchSize,
+			const uint32 MaskSize,
+			const uint32 QueryBufferStride)
 		{
-			/** Virtual destructor */
-			virtual ~ILayerInstance() = default;
+			NNE_RUNTIME_BASIC_TRACE_SCOPE(NNE::RuntimeBasic::Private::OperatorAggregateInsertOneHot);
 
-			/** Indicate to this layer instance what the batchsize is going to be when performing inference. */
-			virtual void SetBatchSize(const uint32 BatchSize) = 0;
-		};
+#if NNE_RUNTIME_BASIC_ENABLE_ISPC
+			ispc::NNERuntimeBasicCPUOperatorAggregateInsertOneHot(
+				QueryBuffer,
+				Index,
+				BatchSize,
+				MaskSize,
+				QueryBufferStride);
+#else
+			for (uint32 BatchIdx = 0; BatchIdx < BatchSize; BatchIdx++)
+			{
+				for (uint32 MaskIdx = 0; MaskIdx < MaskSize; MaskIdx++)
+				{
+					QueryBuffer[BatchIdx * QueryBufferStride + MaskIdx] = 0.0f;
+				}
 
-		/**
-		 * Interface for a Layer - the network parameter data required for a layer.
-		 */
-		struct ILayer
-		{
-			/** Virtual destructor */
-			virtual ~ILayer() = default;
+				QueryBuffer[BatchIdx * QueryBufferStride + Index] = 1.0f;
+			}
+#endif
+		}
 
-			/** Create the instance data required for this type of layer. */
-			virtual TSharedPtr<ILayerInstance> MakeInstance() const { return nullptr; };
+		static inline void OperatorAggregateCountElementNum(
+			uint32& TotalElementNum,
+			uint32* RESTRICT ElementNums,
+			uint32* RESTRICT ElementOffsets,
+			const float* RESTRICT MaskBuffer,
+			const uint32 BatchSize,
+			const uint32 MaskSize,
+			const uint32 MaskBufferStride)
+		{
+			NNE_RUNTIME_BASIC_TRACE_SCOPE(NNE::RuntimeBasic::Private::OperatorAggregateCountElementNum);
 
-			/** Get the layer type. */
-			virtual ELayerType GetLayerType() const = 0;
+#if NNE_RUNTIME_BASIC_ENABLE_ISPC
+			ispc::NNERuntimeBasicCPUOperatorAggregateCountElementNum(
+				TotalElementNum,
+				ElementNums,
+				ElementOffsets,
+				MaskBuffer,
+				BatchSize,
+				MaskSize,
+				MaskBufferStride);
+#else
+			TotalElementNum = 0;
 
-			/** Get the size of the input vector. */
-			virtual uint32 GetInputSize() const = 0;
+			for (uint32 BatchIdx = 0; BatchIdx < BatchSize; BatchIdx++)
+			{
+				uint32 ElementSum = 0;
+				for (uint32 MaskIdx = 0; MaskIdx < MaskSize; MaskIdx++)
+				{
+					if (MaskBuffer[BatchIdx * MaskBufferStride + MaskIdx]) { ElementSum++; }
+				}
 
-			/** Get the size of the output vector. */
-			virtual uint32 GetOutputSize() const = 0;
+				ElementOffsets[BatchIdx] = TotalElementNum;
+				ElementNums[BatchIdx] = ElementSum;
+				TotalElementNum += ElementSum;
+			}
+#endif
+		}
 
-			/** Compute the size required to serialize this layer by growing InOutOffset. */
-			virtual void SerializationSize(uint64& InOutOffset) const = 0;
+		static inline void OperatorAggregateGatherFromSubLayers(
+			float* RESTRICT QueryBuffer,
+			float* RESTRICT KeyBuffer,
+			float* RESTRICT ValueBuffer,
+			uint32* RESTRICT ElementAccum,
+			const uint32* RESTRICT ElementNums,
+			const uint32* RESTRICT ElementOffsets,
+			const TConstArrayView<TArray<uint32>> SubLayerBatchIndices,
+			const TConstArrayView<TArray<float>> SubLayerQueryBuffers,
+			const TConstArrayView<TArray<float>> SubLayerKeyBuffers,
+			const TConstArrayView<TArray<float>> SubLayerValueBuffers,
+			const uint32 BatchSize,
+			const uint32 QuerySize,
+			const uint32 KeySize,
+			const uint32 ValueSize)
+		{
+			NNE_RUNTIME_BASIC_TRACE_SCOPE(NNE::RuntimeBasic::Private::OperatorAggregateGatherFromSubLayers);
 
-			/** Load this layer from the buffer at the given offset. */
-			virtual void SerializationLoad(uint64& InOutOffset, TConstArrayView<uint8> Data) = 0;
+			const uint32 SubLayerNum = SubLayerBatchIndices.Num();
 
-			/** Save this layer from the buffer at the given offset. */
-			virtual void SerializationSave(uint64& InOutOffset, TArrayView<uint8> Data) const = 0;
+			for (uint32 BatchIdx = 0; BatchIdx < BatchSize; BatchIdx++)
+			{
+				ElementAccum[BatchIdx] = 0;
+			}
 
-			/**
-			 * Evaluate this layer.
-			 * 
-			 * Note that the input and output buffers may not match the exact dimensions expected by this layer - the
-			 * size dimension can be over-allocated to allow for buffer re-use.
-			 *
-			 * @param Output			The output buffer of size (BatchSize, OutputStride).
-			 * @param Input				The input buffer of size (BatchSize, InputStride).
-			 * @param Instance			The instance data for this layer.
-			 * @param BatchSize			The batchsize of the inputs and outputs
-			 * @param OutputStride		The stride of the output buffer vector for each item in the batch
-			 * @param InputStride		The stride of the input buffer vector for each item in the batch
-			 */
-			virtual void Evaluate(
-				TArrayView<float> Output,
-				TConstArrayView<float> Input,
-				ILayerInstance* Instance,
-				uint32 BatchSize,
-				uint32 OutputStride,
-				uint32 InputStride) = 0;
-		};
+			for (uint32 SubLayerIdx = 0; SubLayerIdx < SubLayerNum; SubLayerIdx++)
+			{
+				const float* RESTRICT SubLayerQueryBuffer = SubLayerQueryBuffers[SubLayerIdx].GetData();
+				const float* RESTRICT SubLayerKeyBuffer = SubLayerKeyBuffers[SubLayerIdx].GetData();
+				const float* RESTRICT SubLayerValueBuffer = SubLayerValueBuffers[SubLayerIdx].GetData();
+				const uint32* RESTRICT SubLayerBatchIndicesBuffer = SubLayerBatchIndices[SubLayerIdx].GetData();
+				const uint32 SubLayerBatchIndexNum = SubLayerBatchIndices[SubLayerIdx].Num();
 
-		//--------------------------------------------------------------------------
+#if NNE_RUNTIME_BASIC_ENABLE_ISPC
+				ispc::NNERuntimeBasicCPUOperatorAggregateGatherQueryValueFromSubLayers(
+					QueryBuffer,
+					KeyBuffer,
+					ValueBuffer,
+					ElementAccum,
+					ElementOffsets,
+					SubLayerQueryBuffer,
+					SubLayerKeyBuffer,
+					SubLayerValueBuffer,
+					SubLayerBatchIndicesBuffer,
+					SubLayerBatchIndexNum,
+					QuerySize,
+					KeySize,
+					ValueSize);
+#else
+				for (uint32 ElementIdx = 0; ElementIdx < SubLayerBatchIndexNum; ElementIdx++)
+				{
+					const uint32 BatchIdx = SubLayerBatchIndicesBuffer[ElementIdx];
+					const uint32 ElementOffset = ElementOffsets[BatchIdx] + ElementAccum[BatchIdx];
 
-		struct FSequenceLayer;
+					for (uint32 QueryIdx = 0; QueryIdx < QuerySize; QueryIdx++)
+					{
+						QueryBuffer[ElementOffset * QuerySize + QueryIdx] = SubLayerQueryBuffer[ElementIdx * QuerySize + QueryIdx];
+					}
 
-		struct FSequenceLayerInstance : public ILayerInstance
-		{
-			FSequenceLayerInstance(const FSequenceLayer& InSequenceLayer);
+					for (uint32 KeyIdx = 0; KeyIdx < KeySize; KeyIdx++)
+					{
+						KeyBuffer[ElementOffset * KeySize + KeyIdx] = SubLayerKeyBuffer[ElementIdx * KeySize + KeyIdx];
+					}
 
-			virtual void SetBatchSize(const uint32 BatchSize) override final;
+					for (uint32 ValueIdx = 0; ValueIdx < ValueSize; ValueIdx++)
+					{
+						ValueBuffer[ElementOffset * ValueSize + ValueIdx] = SubLayerValueBuffer[ElementIdx * ValueSize + ValueIdx];
+					}
 
-			const FSequenceLayer& SequenceLayer;
-			uint32 ActivationStride = 0;
-			TArray<TSharedPtr<ILayerInstance>, TInlineAllocator<32>> Instances;
-			TArray<float, TInlineAllocator<512>> ActivationBufferFront;
-			TArray<float, TInlineAllocator<512>> ActivationBufferBack;
-		};
+					ElementAccum[BatchIdx]++;
+				}
+#endif
+			}
 
-		//--------------------------------------------------------------------------
+			for (uint32 BatchIdx = 0; BatchIdx < BatchSize; BatchIdx++)
+			{
+				check(ElementAccum[BatchIdx] == ElementNums[BatchIdx]);
+			}
+		}
 
-		struct FSequenceLayer : public ILayer
+		static inline void OperatorAggregateDotProductAttention(
+			float* RESTRICT Attention,
+			const float* RESTRICT Queries,
+			const float* RESTRICT Keys,
+			const uint32 ElementNum,
+			const uint32 AttentionEncodingSize,
+			const uint32 AttentionHeadNum)
 		{
-			virtual TSharedPtr<ILayerInstance> MakeInstance() const { return MakeShared<FSequenceLayerInstance>(*this); };
-			virtual ELayerType GetLayerType() const override final { return ELayerType::Sequence; }
-			virtual uint32 GetInputSize() const override final { return Layers.Num() > 0 ? Layers[0]->GetInputSize() : 0; }
-			virtual uint32 GetOutputSize() const override final { return Layers.Num() > 0 ? Layers.Last()->GetOutputSize() : 0; }
+			NNE_RUNTIME_BASIC_TRACE_SCOPE(NNE::RuntimeBasic::Private::OperatorAggregateDotProductAttention);
 
-			virtual void SerializationSize(uint64& InOutOffset) const override final
+#if NNE_RUNTIME_BASIC_ENABLE_ISPC
+			ispc::NNERuntimeBasicCPUOperatorAggregateDotProductAttention(
+				Attention,
+				Queries,
+				Keys,
+				ElementNum,
+				AttentionEncodingSize,
+				AttentionHeadNum);
+#else
+			for (uint32 ElementIdx = 0; ElementIdx < ElementNum; ElementIdx++)
 			{
-				Serialization::Size(InOutOffset, (uint32)Layers.Num());
-				for (const TSharedPtr<ILayer>& Layer : Layers)
+				for (uint32 HeadIdx = 0; HeadIdx < AttentionHeadNum; HeadIdx++)
 				{
-					Serialization::Size(InOutOffset, Layer);
+					Attention[ElementIdx * AttentionHeadNum + HeadIdx] = 0.0f;
+
+					for (uint32 Idx = 0; Idx < AttentionEncodingSize; Idx++)
+					{
+						Attention[ElementIdx * AttentionHeadNum + HeadIdx] += (
+							Keys[ElementIdx * AttentionHeadNum * AttentionEncodingSize + HeadIdx * AttentionEncodingSize + Idx] *
+							Queries[ElementIdx * AttentionHeadNum * AttentionEncodingSize + HeadIdx * AttentionEncodingSize + Idx]);
+					}
+
+					Attention[ElementIdx * AttentionHeadNum + HeadIdx] /= FMath::Sqrt((float)AttentionEncodingSize);
 				}
 			}
+#endif
+		}
 
-			virtual void SerializationLoad(uint64& InOutOffset, TConstArrayView<uint8> Data) override final
+		static inline void OperatorAggregateSoftmaxPlusOneInplace(
+			float* RESTRICT AttentionMaxs,
+			float* RESTRICT AttentionDenoms,
+			float* RESTRICT Attention,
+			const uint32* RESTRICT ElementNums,
+			const uint32* RESTRICT ElementOffsets,
+			const uint32 BatchSize,
+			const uint32 AttentionHeadNum)
+		{
+			NNE_RUNTIME_BASIC_TRACE_SCOPE(NNE::RuntimeBasic::Private::OperatorAggregateSoftmaxPlusOneInplace);
+
+			// Numerically stable soft-max computation using subtraction of the (positive) max value
+			// 
+			// Here the +1 in the denominator allows the attention 
+			// to attend to nothing as discussed here:
+			// https://www.evanmiller.org/attention-is-off-by-one.html
+
+#if NNE_RUNTIME_BASIC_ENABLE_ISPC
+			ispc::NNERuntimeBasicCPUOperatorAggregateSoftmaxPlusOneInplace(
+				AttentionMaxs,
+				AttentionDenoms,
+				Attention,
+				ElementNums,
+				ElementOffsets,
+				BatchSize,
+				AttentionHeadNum);
+#else
+			for (uint32 BatchIdx = 0; BatchIdx < BatchSize; BatchIdx++)
 			{
-				uint32 LayerNum = 0;
-				Serialization::Load(InOutOffset, LayerNum, Data);
-				Layers.SetNumZeroed(LayerNum);
+				const uint32 ElementNum = ElementNums[BatchIdx];
+				const uint32 ElementOffset = ElementOffsets[BatchIdx];
 
-				for (uint32 LayerIdx = 0; LayerIdx < LayerNum; LayerIdx++)
+				for (uint32 HeadIdx = 0; HeadIdx < AttentionHeadNum; HeadIdx++)
 				{
-					Serialization::Load(InOutOffset, Layers[LayerIdx], Data);
+					AttentionMaxs[HeadIdx] = 0.0f;
+					AttentionDenoms[HeadIdx] = 0.0f;
 				}
-			}
 
-			virtual void SerializationSave(uint64& InOutOffset, TArrayView<uint8> Data) const override final
-			{
-				uint32 LayerNum = Layers.Num();
-				Serialization::Save(InOutOffset, LayerNum, Data);
-				for (uint32 LayerIdx = 0; LayerIdx < LayerNum; LayerIdx++)
+				for (uint32 ElementIdx = ElementOffset; ElementIdx < ElementOffset + ElementNum; ElementIdx++)
 				{
-					Serialization::Save(InOutOffset, Layers[LayerIdx], Data);
+					for (uint32 HeadIdx = 0; HeadIdx < AttentionHeadNum; HeadIdx++)
+					{
+						AttentionMaxs[HeadIdx] = FMath::Max(AttentionMaxs[HeadIdx], Attention[ElementIdx * AttentionHeadNum + HeadIdx]);
+					}
 				}
-			}
-
-			virtual void Evaluate(
-				TArrayView<float> Output,
-				TConstArrayView<float> Input,
-				ILayerInstance* Instance,
-				uint32 BatchSize,
-				uint32 OutputStride,
-				uint32 InputStride) override final
-			{
-				TRACE_CPUPROFILER_EVENT_SCOPE(NNE::RuntimeBasic::Private::FSequenceLayer::Evaluate);
-
-				const uint32 LayerNum = Layers.Num();
-				check(LayerNum > 0);
-
-				FSequenceLayerInstance* SequenceInstance = StaticCast<FSequenceLayerInstance*>(Instance);
-				check(SequenceInstance);
-
-				// If we just have one layer then evaluate layer directly without using intermediate storage
 
-				if (LayerNum == 1)
+				for (uint32 ElementIdx = ElementOffset; ElementIdx < ElementOffset + ElementNum; ElementIdx++)
 				{
-					Layers[0]->Evaluate(
-						Output,
-						Input,
-						SequenceInstance->Instances[0].Get(),
-						BatchSize,
-						OutputStride,
-						InputStride);
+					for (uint32 HeadIdx = 0; HeadIdx < AttentionHeadNum; HeadIdx++)
+					{
+						AttentionDenoms[HeadIdx] += FMath::Exp(Attention[ElementIdx * AttentionHeadNum + HeadIdx] - AttentionMaxs[HeadIdx]);
+					}
+				}
 
-					return;
+				for (uint32 ElementIdx = ElementOffset; ElementIdx < ElementOffset + ElementNum; ElementIdx++)
+				{
+					for (uint32 HeadIdx = 0; HeadIdx < AttentionHeadNum; HeadIdx++)
+					{
+						Attention[ElementIdx * AttentionHeadNum + HeadIdx] = FMath::Exp(Attention[ElementIdx * AttentionHeadNum + HeadIdx] - AttentionMaxs[HeadIdx]) / (AttentionDenoms[HeadIdx] + 1.0f);
+					}
 				}
+			}
+#endif
+		}
 
-				// Otherwise evaluate first layer from input into activation buffer
+		static inline void OperatorAggregateAttentionSum(
+			float* RESTRICT Output,
+			const float* RESTRICT Attention,
+			const float* RESTRICT Values,
+			const uint32* RESTRICT ElementNums,
+			const uint32* RESTRICT ElementOffsets,
+			const uint32 BatchSize,
+			const uint32 EncodingSize,
+			const uint32 AttentionHeadNum,
+			const uint32 OutputStride)
+		{
+			NNE_RUNTIME_BASIC_TRACE_SCOPE(NNE::RuntimeBasic::Private::OperatorAggregateAttentionSum);
+
+#if NNE_RUNTIME_BASIC_ENABLE_ISPC
+			ispc::NNERuntimeBasicCPUOperatorAggregateAttentionSum(
+				Output,
+				Attention,
+				Values,
+				ElementNums,
+				ElementOffsets,
+				BatchSize,
+				EncodingSize,
+				AttentionHeadNum,
+				OutputStride);
+#else
+			for (uint32 BatchIdx = 0; BatchIdx < BatchSize; BatchIdx++)
+			{
+				for (uint32 HeadIdx = 0; HeadIdx < AttentionHeadNum; HeadIdx++)
+				{
+					for (uint32 Idx = 0; Idx < EncodingSize; Idx++)
+					{
+						Output[BatchIdx * OutputStride + HeadIdx * EncodingSize + Idx] = 0.0f;
+					}
+				}
+
+				const uint32 ElementNum = ElementNums[BatchIdx];
+				const uint32 ElementOffset = ElementOffsets[BatchIdx];
+
+				for (uint32 ElementIdx = ElementOffset; ElementIdx < ElementOffset + ElementNum; ElementIdx++)
+				{
+					for (uint32 HeadIdx = 0; HeadIdx < AttentionHeadNum; HeadIdx++)
+					{
+						const float Scale = Attention[ElementIdx * AttentionHeadNum + HeadIdx];
+
+						if (Scale != 0.0f)
+						{
+							for (uint32 Idx = 0; Idx < EncodingSize; Idx++)
+							{
+								Output[BatchIdx * OutputStride + HeadIdx * EncodingSize + Idx] += Scale * Values[ElementIdx * AttentionHeadNum * EncodingSize + HeadIdx * EncodingSize + Idx];
+							}
+						}
+					}
+				}
+			}
+#endif
+		}
+
+		static inline void OperatorGather(
+			float* RESTRICT OutputBuffer,
+			const float* RESTRICT InputBuffer,
+			const uint32* RESTRICT BatchIndices,
+			const uint32 BatchIndexNum,
+			const uint32 InputOutputSize,
+			const uint32 OutputStride,
+			const uint32 InputStride)
+		{
+			NNE_RUNTIME_BASIC_TRACE_SCOPE(NNE::RuntimeBasic::Private::OperatorGather);
+
+#if NNE_RUNTIME_BASIC_ENABLE_ISPC
+			ispc::NNERuntimeBasicCPUOperatorGather(
+				OutputBuffer,
+				InputBuffer,
+				BatchIndices,
+				BatchIndexNum,
+				InputOutputSize,
+				OutputStride,
+				InputStride);
+#else
+			for (uint32 BatchIndexIdx = 0; BatchIndexIdx < BatchIndexNum; BatchIndexIdx++)
+			{
+				const uint32 SrcBatchIdx = BatchIndices[BatchIndexIdx];
+				for (uint32 Idx = 0; Idx < InputOutputSize; Idx++)
+				{
+					OutputBuffer[BatchIndexIdx * OutputStride + Idx] = InputBuffer[SrcBatchIdx * InputStride + Idx];
+				}
+			}
+#endif
+		}
+
+		static inline void OperatorScatter(
+			float* RESTRICT OutputBuffer,
+			const float* RESTRICT InputBuffer,
+			const uint32* RESTRICT BatchIndices,
+			const uint32 BatchIndexNum,
+			const uint32 InputOutputSize,
+			const uint32 OutputStride,
+			const uint32 InputStride)
+		{
+			NNE_RUNTIME_BASIC_TRACE_SCOPE(NNE::RuntimeBasic::Private::OperatorScatter);
+
+#if NNE_RUNTIME_BASIC_ENABLE_ISPC
+			ispc::NNERuntimeBasicCPUOperatorScatter(
+				OutputBuffer,
+				InputBuffer,
+				BatchIndices,
+				BatchIndexNum,
+				InputOutputSize,
+				OutputStride,
+				InputStride);
+#else
+			for (uint32 BatchIndexIdx = 0; BatchIndexIdx < BatchIndexNum; BatchIndexIdx++)
+			{
+				const uint32 DstBatchIdx = BatchIndices[BatchIndexIdx];
+				for (uint32 Idx = 0; Idx < InputOutputSize; Idx++)
+				{
+					OutputBuffer[DstBatchIdx * OutputStride + Idx] = InputBuffer[BatchIndexIdx * InputStride + Idx];
+				}
+			}
+#endif
+		}
+
+		static inline void OperatorGatherSubLayerBatchIndicesExclusive(
+			TArrayView<TArray<uint32>> SubLayerBatchIndices,
+			const float* RESTRICT SubLayerMaskBuffer,
+			const uint32 BatchSize,
+			const uint32 SubLayerMaskSize,
+			const uint32 SubLayerMaskStride)
+		{
+			NNE_RUNTIME_BASIC_TRACE_SCOPE(NNE::RuntimeBasic::Private::OperatorGatherSubLayerBatchIndicesExclusive);
+
+			for (uint32 SubLayerIdx = 0; SubLayerIdx < SubLayerMaskSize; SubLayerIdx++)
+			{
+				SubLayerBatchIndices[SubLayerIdx].Reset();
+			}
+
+			for (uint32 BatchIdx = 0; BatchIdx < BatchSize; BatchIdx++)
+			{
+				bool bFound = false;
+				for (uint32 SubLayerIdx = 0; SubLayerIdx < SubLayerMaskSize; SubLayerIdx++)
+				{
+					if (SubLayerMaskBuffer[BatchIdx * SubLayerMaskStride + SubLayerIdx])
+					{
+						SubLayerBatchIndices[SubLayerIdx].Add(BatchIdx);
+						bFound = true;
+						break;
+					}
+				}
+
+				checkf(bFound, TEXT("SubLayer index not found."));
+			}
+		}
+
+		static inline void OperatorGatherSubLayerBatchIndicesInclusive(
+			TArrayView<TArray<uint32>> SubLayerBatchIndices,
+			const float* RESTRICT SubLayerMaskBuffer,
+			const uint32 BatchSize,
+			const uint32 SubLayerMaskSize,
+			const uint32 SubLayerMaskStride)
+		{
+			NNE_RUNTIME_BASIC_TRACE_SCOPE(NNE::RuntimeBasic::Private::OperatorGatherSubLayerBatchIndicesInclusive);
+
+			for (uint32 SubLayerIdx = 0; SubLayerIdx < SubLayerMaskSize; SubLayerIdx++)
+			{
+				SubLayerBatchIndices[SubLayerIdx].Reset();
+			}
+
+			for (uint32 BatchIdx = 0; BatchIdx < BatchSize; BatchIdx++)
+			{
+				for (uint32 SubLayerIdx = 0; SubLayerIdx < SubLayerMaskSize; SubLayerIdx++)
+				{
+					if (SubLayerMaskBuffer[BatchIdx * SubLayerMaskStride + SubLayerIdx])
+					{
+						SubLayerBatchIndices[SubLayerIdx].Add(BatchIdx);
+					}
+				}
+			}
+		}
+
+		//--------------------------------------------------------------------------
+		// Layer Types
+		//--------------------------------------------------------------------------
+
+		/** Layer Type Id - this should match what is given in nne_runtime_basic_cpu.py  */
+		enum class ELayerType : uint32
+		{
+			Invalid = 0,
+			
+			Sequence = 1,
+			
+			Normalize = 2,
+			Denormalize = 3,
+			Linear = 4,
+			CompressedLinear = 5,
+			MultiLinear = 6,
+
+			ReLU = 7,
+			ELU = 8,
+			TanH = 9,
+			PReLU = 10,
+
+			MemoryCell = 11,
+
+			Copy = 12,
+			Concat = 13,
+			Array = 14,
+			
+			AggregateSet = 15,
+			AggregateOrExclusive = 16,
+			AggregateOrInclusive = 17,
+		};
+
+		//--------------------------------------------------------------------------
+		// Layer Type Interfaces
+		//--------------------------------------------------------------------------
+
+		/**
+		 * Interface for a Layer Instance - the data required for performing inference for a layer.
+		 */
+		struct ILayerInstance
+		{
+			/** Virtual destructor */
+			virtual ~ILayerInstance() = default;
+
+			/** Indicate to this layer instance what the maximum batchsize is going to be when performing inference. */
+			virtual void SetMaxBatchSize(const uint32 MaxBatchSize) = 0;
+		};
+
+		/**
+		 * Interface for a Layer - the network parameter data required for a layer.
+		 */
+		struct ILayer
+		{
+			/** Virtual destructor */
+			virtual ~ILayer() = default;
+
+			/** Create the instance data required for this type of layer. */
+			virtual TSharedPtr<ILayerInstance> MakeInstance() const { return nullptr; };
+
+			/** Get the layer type. */
+			virtual ELayerType GetLayerType() const = 0;
+
+			/** Get the size of the input vector. */
+			virtual uint32 GetInputSize() const = 0;
+
+			/** Get the size of the output vector. */
+			virtual uint32 GetOutputSize() const = 0;
+
+			/** Compute the size required to serialize this layer by growing InOutOffset. */
+			virtual void SerializationSize(uint64& InOutOffset) const = 0;
+
+			/** Load this layer from the buffer at the given offset. */
+			virtual void SerializationLoad(uint64& InOutOffset, TConstArrayView<uint8> Data) = 0;
+
+			/** Save this layer from the buffer at the given offset. */
+			virtual void SerializationSave(uint64& InOutOffset, TArrayView<uint8> Data) const = 0;
+
+			/**
+			 * Evaluate this layer.
+			 *
+			 * @param Instance				The instance data for this layer - what was returned by `MakeInstance`.
+			 * @param OutputBuffer			The output buffer
+			 * @param InputBuffer			The input buffer
+			 * @param BatchSize				The number of items in the batch
+			 * @param OutputBufferSize		The vector size of the items in the output
+			 * @param InputBufferSize		The vector size of the items in the input
+			 * @param OutputBufferStride	The stride of the output for each item in the batch
+			 * @param InputBufferStride		The stride of the input for each item in the batch
+			 */
+			virtual void Evaluate(
+				ILayerInstance* Instance,
+				float* OutputBuffer,
+				const float* InputBuffer,
+				const uint32 BatchSize,
+				const uint32 OutputBufferSize,
+				const uint32 InputBufferSize,
+				const uint32 OutputBufferStride,
+				const uint32 InputBufferStride) = 0;
+		};
+
+		//--------------------------------------------------------------------------
+		// Layers
+		//--------------------------------------------------------------------------
+
+		struct FSequenceLayer;
+
+		struct FSequenceLayerInstance : public ILayerInstance
+		{
+			FSequenceLayerInstance(const FSequenceLayer& InSequenceLayer);
+
+			virtual void SetMaxBatchSize(const uint32 MaxBatchSize) override final;
+
+			const FSequenceLayer& SequenceLayer;
+			uint32 ActivationStride = 0;
+			TArray<TSharedPtr<ILayerInstance>, TInlineAllocator<32>> Instances;
+			TArray<float> ActivationBufferFront;
+			TArray<float> ActivationBufferBack;
+		};
+
+		struct FSequenceLayer : public ILayer
+		{
+			virtual TSharedPtr<ILayerInstance> MakeInstance() const { return MakeShared<FSequenceLayerInstance>(*this); };
+			virtual ELayerType GetLayerType() const override final { return ELayerType::Sequence; }
+			virtual uint32 GetInputSize() const override final { return Layers.Num() > 0 ? Layers[0]->GetInputSize() : 0; }
+			virtual uint32 GetOutputSize() const override final { return Layers.Num() > 0 ? Layers.Last()->GetOutputSize() : 0; }
+
+			virtual void SerializationSize(uint64& InOutOffset) const override final
+			{
+				Serialization::Size(InOutOffset, (uint32)Layers.Num());
+				Serialization::Size(InOutOffset, Layers);
+			}
+
+			virtual void SerializationLoad(uint64& InOutOffset, TConstArrayView<uint8> Data) override final
+			{
+				uint32 LayerNum = 0;
+				Serialization::Load(InOutOffset, LayerNum, Data);
+				Layers.Init(nullptr, LayerNum);
+				Serialization::Load(InOutOffset, Layers, Data);
+
+				LayerInputSizes.SetNumUninitialized(LayerNum);
+				LayerOutputSizes.SetNumUninitialized(LayerNum);
+				for (uint32 LayerIdx = 0; LayerIdx < LayerNum; LayerIdx++)
+				{
+					LayerInputSizes[LayerIdx] = Layers[LayerIdx]->GetInputSize();
+					LayerOutputSizes[LayerIdx] = Layers[LayerIdx]->GetOutputSize();
+				}
+			}
+
+			virtual void SerializationSave(uint64& InOutOffset, TArrayView<uint8> Data) const override final
+			{
+				Serialization::Save(InOutOffset, (uint32)Layers.Num(), Data);
+				Serialization::Save(InOutOffset, Layers, Data);
+			}
+
+			virtual void Evaluate(
+				ILayerInstance* Instance,
+				float* OutputBuffer,
+				const float* InputBuffer,
+				const uint32 BatchSize,
+				const uint32 OutputBufferSize,
+				const uint32 InputBufferSize,
+				const uint32 OutputBufferStride,
+				const uint32 InputBufferStride) override final
+			{
+				NNE_RUNTIME_BASIC_TRACE_SCOPE(NNE::RuntimeBasic::Private::FSequenceLayer::Evaluate);
+				check(OutputBufferSize == GetOutputSize() && InputBufferSize == GetInputSize());
+				OperatorNanCheck(InputBuffer, BatchSize, InputBufferSize, InputBufferStride);
+
+				const uint32 LayerNum = Layers.Num();
+				check(LayerNum > 0);
+
+				FSequenceLayerInstance* SequenceInstance = StaticCast<FSequenceLayerInstance*>(Instance);
+				check(SequenceInstance);
+
+				// If we just have one layer then evaluate layer directly without using intermediate storage
+
+				if (LayerNum == 1)
+				{
+					Layers[0]->Evaluate(
+						SequenceInstance->Instances[0].Get(),
+						OutputBuffer,
+						InputBuffer,
+						BatchSize,
+						OutputBufferSize,
+						InputBufferSize,
+						OutputBufferStride,
+						InputBufferStride);
+
+					return;
+				}
+
+				// Otherwise evaluate first layer from input into activation buffer
 
 				Layers[0]->Evaluate(
-					SequenceInstance->ActivationBufferFront,
-					Input,
 					SequenceInstance->Instances[0].Get(),
+					SequenceInstance->ActivationBufferFront.GetData(),
+					InputBuffer,
 					BatchSize,
+					LayerOutputSizes[0],
+					LayerInputSizes[0],
 					SequenceInstance->ActivationStride,
-					InputStride);
+					InputBufferStride);
 
 				// Evaluate intermediate layers using front and back buffers
 
@@ -725,10 +1329,12 @@ namespace UE::NNE::RuntimeBasic
 						SequenceInstance->ActivationBufferBack;
 
 					Layers[LayerIdx]->Evaluate(
-						LayerOutput,
-						LayerInput,
 						SequenceInstance->Instances[LayerIdx].Get(),
+						LayerOutput.GetData(),
+						LayerInput.GetData(),
 						BatchSize,
+						LayerOutputSizes[LayerIdx],
+						LayerInputSizes[LayerIdx],
 						SequenceInstance->ActivationStride,
 						SequenceInstance->ActivationStride);
 				}
@@ -740,19 +1346,23 @@ namespace UE::NNE::RuntimeBasic
 					SequenceInstance->ActivationBufferBack;
 
 				Layers.Last()->Evaluate(
-					Output,
-					FinalLayerInput,
 					SequenceInstance->Instances.Last().Get(),
+					OutputBuffer,
+					FinalLayerInput.GetData(),
 					BatchSize,
-					OutputStride,
+					OutputBufferSize,
+					LayerInputSizes.Last(),
+					OutputBufferStride,
 					SequenceInstance->ActivationStride);
+
+				OperatorNanCheck(OutputBuffer, BatchSize, OutputBufferSize, OutputBufferStride);
 			}
 
 			TArray<TSharedPtr<ILayer>, TInlineAllocator<32>> Layers;
+			TArray<uint32, TInlineAllocator<32>> LayerInputSizes;
+			TArray<uint32, TInlineAllocator<32>> LayerOutputSizes;
 		};
 
-		//--------------------------------------------------------------------------
-
 		FSequenceLayerInstance::FSequenceLayerInstance(const FSequenceLayer& InSequenceLayer) 
 			: SequenceLayer(InSequenceLayer)
 		{
@@ -773,9 +1383,9 @@ namespace UE::NNE::RuntimeBasic
 			}
 		}
 
-		void FSequenceLayerInstance::SetBatchSize(const uint32 BatchSize)
+		void FSequenceLayerInstance::SetMaxBatchSize(const uint32 MaxBatchSize)
 		{
-			TRACE_CPUPROFILER_EVENT_SCOPE(NNE::RuntimeBasic::Private::FSequenceLayerInstance::SetBatchSize);
+			NNE_RUNTIME_BASIC_TRACE_SCOPE(NNE::RuntimeBasic::Private::FSequenceLayerInstance::SetMaxBatchSize);
 
 			// Propagate call to sub-layer instances
 
@@ -784,15 +1394,15 @@ namespace UE::NNE::RuntimeBasic
 				// Most layers don't allocate instance data and so we need to check for nullptr
 				if (Instance)
 				{
-					Instance->SetBatchSize(BatchSize);
+					Instance->SetMaxBatchSize(MaxBatchSize);
 				}
 			}
 
 			// Allocate front and back buffers to maximum size. Don't shrink to avoid re-allocation 
 			// when smaller batches are requested.
 
-			ActivationBufferFront.SetNumUninitialized(BatchSize * ActivationStride, false);
-			ActivationBufferBack.SetNumUninitialized(BatchSize * ActivationStride, false);
+			ActivationBufferFront.SetNumUninitialized(MaxBatchSize * ActivationStride, false);
+			ActivationBufferBack.SetNumUninitialized(MaxBatchSize * ActivationStride, false);
 		}
 
 		//--------------------------------------------------------------------------
@@ -825,24 +1435,31 @@ namespace UE::NNE::RuntimeBasic
 			}
 
 			virtual void Evaluate(
-				TArrayView<float> Output,
-				TConstArrayView<float> Input,
 				ILayerInstance* Instance,
-				uint32 BatchSize,
-				uint32 OutputStride,
-				uint32 InputStride) override final
-			{
-				TRACE_CPUPROFILER_EVENT_SCOPE(NNE::RuntimeBasic::Private::FNormalizeLayer::Evaluate);
-
-				NormalizeLayer(
-					Output.GetData(),
-					Input.GetData(),
+				float* OutputBuffer,
+				const float* InputBuffer,
+				const uint32 BatchSize,
+				const uint32 OutputBufferSize,
+				const uint32 InputBufferSize,
+				const uint32 OutputBufferStride,
+				const uint32 InputBufferStride) override final
+			{
+				NNE_RUNTIME_BASIC_TRACE_SCOPE(NNE::RuntimeBasic::Private::FNormalizeLayer::Evaluate);
+				check(OutputBufferSize == GetOutputSize() && InputBufferSize == GetInputSize());
+				check(Instance == nullptr);
+				OperatorNanCheck(InputBuffer, BatchSize, InputBufferSize, InputBufferStride);
+
+				OperatorNormalize(
+					OutputBuffer,
+					InputBuffer,
 					Mean.GetData(),
 					Std.GetData(),
 					BatchSize,
 					InputOutputSize,
-					OutputStride,
-					InputStride);
+					OutputBufferStride,
+					InputBufferStride);
+
+				OperatorNanCheck(OutputBuffer, BatchSize, OutputBufferSize, OutputBufferStride);
 			}
 
 			uint32 InputOutputSize = 0;
@@ -852,8 +1469,6 @@ namespace UE::NNE::RuntimeBasic
 
 		//--------------------------------------------------------------------------
 
-
-
 		struct FDenormalizeLayer : public ILayer
 		{
 			virtual ELayerType GetLayerType() const override final { return ELayerType::Denormalize; }
@@ -882,24 +1497,31 @@ namespace UE::NNE::RuntimeBasic
 			}
 
 			virtual void Evaluate(
-				TArrayView<float> Output,
-				TConstArrayView<float> Input,
 				ILayerInstance* Instance,
-				uint32 BatchSize,
-				uint32 OutputStride,
-				uint32 InputStride) override final
-			{
-				TRACE_CPUPROFILER_EVENT_SCOPE(NNE::RuntimeBasic::Private::FDenormalizeLayer::Evaluate);
-
-				DenormalizeLayer(
-					Output.GetData(),
-					Input.GetData(),
+				float* OutputBuffer,
+				const float* InputBuffer,
+				const uint32 BatchSize,
+				const uint32 OutputBufferSize,
+				const uint32 InputBufferSize,
+				const uint32 OutputBufferStride,
+				const uint32 InputBufferStride) override final
+			{
+				NNE_RUNTIME_BASIC_TRACE_SCOPE(NNE::RuntimeBasic::Private::FDenormalizeLayer::Evaluate);
+				check(OutputBufferSize == GetOutputSize() && InputBufferSize == GetInputSize());
+				check(Instance == nullptr);
+				OperatorNanCheck(InputBuffer, BatchSize, InputBufferSize, InputBufferStride);
+
+				OperatorDenormalize(
+					OutputBuffer,
+					InputBuffer,
 					Mean.GetData(),
 					Std.GetData(),
 					BatchSize,
 					InputOutputSize,
-					OutputStride,
-					InputStride);
+					OutputBufferStride,
+					InputBufferStride);
+
+				OperatorNanCheck(OutputBuffer, BatchSize, OutputBufferSize, OutputBufferStride);
 			}
 
 			uint32 InputOutputSize = 0;
@@ -940,25 +1562,32 @@ namespace UE::NNE::RuntimeBasic
 			}
 
 			virtual void Evaluate(
-				TArrayView<float> Output,
-				TConstArrayView<float> Input,
 				ILayerInstance* Instance,
-				uint32 BatchSize,
-				uint32 OutputStride,
-				uint32 InputStride) override final
-			{
-				TRACE_CPUPROFILER_EVENT_SCOPE(NNE::RuntimeBasic::Private::FLinearLayer::Evaluate);
-
-				LinearLayer(
-					Output.GetData(),
-					Input.GetData(),
+				float* OutputBuffer,
+				const float* InputBuffer,
+				const uint32 BatchSize,
+				const uint32 OutputBufferSize,
+				const uint32 InputBufferSize,
+				const uint32 OutputBufferStride,
+				const uint32 InputBufferStride) override final
+			{
+				NNE_RUNTIME_BASIC_TRACE_SCOPE(NNE::RuntimeBasic::Private::FLinearLayer::Evaluate);
+				check(OutputBufferSize == GetOutputSize() && InputBufferSize == GetInputSize());
+				check(Instance == nullptr);
+				OperatorNanCheck(InputBuffer, BatchSize, InputBufferSize, InputBufferStride);
+
+				OperatorLinear(
+					OutputBuffer,
+					InputBuffer,
 					Weights.GetData(),
 					Biases.GetData(),
 					BatchSize,
-					InputSize,
 					OutputSize,
-					OutputStride,
-					InputStride);
+					InputSize,
+					OutputBufferStride,
+					InputBufferStride);
+				
+				OperatorNanCheck(OutputBuffer, BatchSize, OutputBufferSize, OutputBufferStride);
 			}
 
 			uint32 InputSize = 0;
@@ -1006,27 +1635,34 @@ namespace UE::NNE::RuntimeBasic
 			}
 
 			virtual void Evaluate(
-				TArrayView<float> Output,
-				TConstArrayView<float> Input,
 				ILayerInstance* Instance,
-				uint32 BatchSize,
-				uint32 OutputStride,
-				uint32 InputStride) override final
-			{
-				TRACE_CPUPROFILER_EVENT_SCOPE(NNE::RuntimeBasic::Private::FCompressedLinearLayer::Evaluate);
-
-				CompressedLinearLayer(
-					Output.GetData(),
-					Input.GetData(),
+				float* OutputBuffer,
+				const float* InputBuffer,
+				const uint32 BatchSize,
+				const uint32 OutputBufferSize,
+				const uint32 InputBufferSize,
+				const uint32 OutputBufferStride,
+				const uint32 InputBufferStride) override final
+			{
+				NNE_RUNTIME_BASIC_TRACE_SCOPE(NNE::RuntimeBasic::Private::FCompressedLinearLayer::Evaluate);
+				check(OutputBufferSize == GetOutputSize() && InputBufferSize == GetInputSize());
+				check(Instance == nullptr);
+				OperatorNanCheck(InputBuffer, BatchSize, InputBufferSize, InputBufferStride);
+
+				OperatorCompressedLinear(
+					OutputBuffer,
+					InputBuffer,
 					Weights.GetData(),
 					WeightOffsets.GetData(),
 					WeightScales.GetData(),
 					Biases.GetData(),
 					BatchSize,
-					InputSize,
 					OutputSize,
-					OutputStride,
-					InputStride);
+					InputSize,
+					OutputBufferStride,
+					InputBufferStride);
+
+				OperatorNanCheck(OutputBuffer, BatchSize, OutputBufferSize, OutputBufferStride);
 			}
 
 			uint32 InputSize = 0;
@@ -1073,26 +1709,33 @@ namespace UE::NNE::RuntimeBasic
 			}
 
 			virtual void Evaluate(
-				TArrayView<float> Output,
-				TConstArrayView<float> Input,
 				ILayerInstance* Instance,
-				uint32 BatchSize,
-				uint32 OutputStride,
-				uint32 InputStride) override final
-			{
-				TRACE_CPUPROFILER_EVENT_SCOPE(NNE::RuntimeBasic::Private::FMultiLinearLayer::Evaluate);
-
-				MultiLinearLayer(
-					Output.GetData(),
-					Input.GetData(),
+				float* OutputBuffer,
+				const float* InputBuffer,
+				const uint32 BatchSize,
+				const uint32 OutputBufferSize,
+				const uint32 InputBufferSize,
+				const uint32 OutputBufferStride,
+				const uint32 InputBufferStride) override final
+			{
+				NNE_RUNTIME_BASIC_TRACE_SCOPE(NNE::RuntimeBasic::Private::FMultiLinearLayer::Evaluate);
+				check(OutputBufferSize == GetOutputSize() && InputBufferSize == GetInputSize());
+				check(Instance == nullptr);
+				OperatorNanCheck(InputBuffer, BatchSize, InputBufferSize, InputBufferStride);
+
+				OperatorMultiLinear(
+					OutputBuffer,
+					InputBuffer,
 					Weights.GetData(),
 					Biases.GetData(),
 					BatchSize,
 					BlockNum,
-					InputSize,
 					OutputSize,
-					OutputStride,
-					InputStride);
+					InputSize,
+					OutputBufferStride,
+					InputBufferStride);
+
+				OperatorNanCheck(OutputBuffer, BatchSize, OutputBufferSize, OutputBufferStride);
 			}
 
 			uint32 InputSize = 0;
@@ -1104,8 +1747,6 @@ namespace UE::NNE::RuntimeBasic
 
 		//--------------------------------------------------------------------------
 
-
-
 		struct FReLULayer : public ILayer
 		{
 			virtual ELayerType GetLayerType() const override final { return ELayerType::ReLU; }
@@ -1128,22 +1769,29 @@ namespace UE::NNE::RuntimeBasic
 			}
 
 			virtual void Evaluate(
-				TArrayView<float> Output,
-				TConstArrayView<float> Input,
 				ILayerInstance* Instance,
-				uint32 BatchSize,
-				uint32 OutputStride,
-				uint32 InputStride) override final
-			{
-				TRACE_CPUPROFILER_EVENT_SCOPE(NNE::RuntimeBasic::Private::FReLULayer::Evaluate);
-
-				ActivationReLU(
-					Output.GetData(),
-					Input.GetData(),
+				float* OutputBuffer,
+				const float* InputBuffer,
+				const uint32 BatchSize,
+				const uint32 OutputBufferSize,
+				const uint32 InputBufferSize,
+				const uint32 OutputBufferStride,
+				const uint32 InputBufferStride) override final
+			{
+				NNE_RUNTIME_BASIC_TRACE_SCOPE(NNE::RuntimeBasic::Private::FReLULayer::Evaluate);
+				check(OutputBufferSize == GetOutputSize() && InputBufferSize == GetInputSize());
+				check(Instance == nullptr);
+				OperatorNanCheck(InputBuffer, BatchSize, InputBufferSize, InputBufferStride);
+
+				OperatorReLU(
+					OutputBuffer,
+					InputBuffer,
 					BatchSize,
 					InputOutputSize,
-					OutputStride,
-					InputStride);
+					OutputBufferStride,
+					InputBufferStride);
+
+				OperatorNanCheck(OutputBuffer, BatchSize, OutputBufferSize, OutputBufferStride);
 			}
 
 			uint32 InputOutputSize = 0;
@@ -1173,22 +1821,29 @@ namespace UE::NNE::RuntimeBasic
 			}
 
 			virtual void Evaluate(
-				TArrayView<float> Output,
-				TConstArrayView<float> Input,
 				ILayerInstance* Instance,
-				uint32 BatchSize,
-				uint32 OutputStride,
-				uint32 InputStride) override final
-			{
-				TRACE_CPUPROFILER_EVENT_SCOPE(NNE::RuntimeBasic::Private::FELULayer::Evaluate);
-
-				ActivationELU(
-					Output.GetData(),
-					Input.GetData(),
+				float* OutputBuffer,
+				const float* InputBuffer,
+				const uint32 BatchSize,
+				const uint32 OutputBufferSize,
+				const uint32 InputBufferSize,
+				const uint32 OutputBufferStride,
+				const uint32 InputBufferStride) override final
+			{
+				NNE_RUNTIME_BASIC_TRACE_SCOPE(NNE::RuntimeBasic::Private::FELULayer::Evaluate);
+				check(OutputBufferSize == GetOutputSize() && InputBufferSize == GetInputSize());
+				check(Instance == nullptr);
+				OperatorNanCheck(InputBuffer, BatchSize, InputBufferSize, InputBufferStride);
+
+				OperatorELU(
+					OutputBuffer,
+					InputBuffer,
 					BatchSize,
 					InputOutputSize,
-					OutputStride,
-					InputStride);
+					OutputBufferStride,
+					InputBufferStride);
+
+				OperatorNanCheck(OutputBuffer, BatchSize, OutputBufferSize, OutputBufferStride);
 			}
 
 			uint32 InputOutputSize = 0;
@@ -1218,22 +1873,29 @@ namespace UE::NNE::RuntimeBasic
 			}
 
 			virtual void Evaluate(
-				TArrayView<float> Output,
-				TConstArrayView<float> Input,
 				ILayerInstance* Instance,
-				uint32 BatchSize,
-				uint32 OutputStride,
-				uint32 InputStride) override final
-			{
-				TRACE_CPUPROFILER_EVENT_SCOPE(NNE::RuntimeBasic::Private::FTanHLayer::Evaluate);
-
-				ActivationTanH(
-					Output.GetData(),
-					Input.GetData(),
+				float* OutputBuffer,
+				const float* InputBuffer,
+				const uint32 BatchSize,
+				const uint32 OutputBufferSize,
+				const uint32 InputBufferSize,
+				const uint32 OutputBufferStride,
+				const uint32 InputBufferStride) override final
+			{
+				NNE_RUNTIME_BASIC_TRACE_SCOPE(NNE::RuntimeBasic::Private::FTanHLayer::Evaluate);
+				check(OutputBufferSize == GetOutputSize() && InputBufferSize == GetInputSize());
+				check(Instance == nullptr);
+				OperatorNanCheck(InputBuffer, BatchSize, InputBufferSize, InputBufferStride);
+
+				OperatorTanH(
+					OutputBuffer,
+					InputBuffer,
 					BatchSize,
 					InputOutputSize,
-					OutputStride,
-					InputStride);
+					OutputBufferStride,
+					InputBufferStride);
+
+				OperatorNanCheck(OutputBuffer, BatchSize, OutputBufferSize, OutputBufferStride);
 			}
 
 			uint32 InputOutputSize = 0;
@@ -1266,23 +1928,30 @@ namespace UE::NNE::RuntimeBasic
 			}
 
 			virtual void Evaluate(
-				TArrayView<float> Output,
-				TConstArrayView<float> Input,
 				ILayerInstance* Instance,
-				uint32 BatchSize,
-				uint32 OutputStride,
-				uint32 InputStride) override final
-			{
-				TRACE_CPUPROFILER_EVENT_SCOPE(NNE::RuntimeBasic::Private::FPReLuLayer::Evaluate);
-
-				ActivationPReLU(
-					Output.GetData(),
-					Input.GetData(),
+				float* OutputBuffer,
+				const float* InputBuffer,
+				const uint32 BatchSize,
+				const uint32 OutputBufferSize,
+				const uint32 InputBufferSize,
+				const uint32 OutputBufferStride,
+				const uint32 InputBufferStride) override final
+			{
+				NNE_RUNTIME_BASIC_TRACE_SCOPE(NNE::RuntimeBasic::Private::FPReLuLayer::Evaluate);
+				check(OutputBufferSize == GetOutputSize() && InputBufferSize == GetInputSize());
+				check(Instance == nullptr);
+				OperatorNanCheck(InputBuffer, BatchSize, InputBufferSize, InputBufferStride);
+
+				OperatorPReLU(
+					OutputBuffer,
+					InputBuffer,
 					Alpha.GetData(),
 					BatchSize,
 					InputOutputSize,
-					OutputStride,
-					InputStride);
+					OutputBufferStride,
+					InputBufferStride);
+
+				OperatorNanCheck(OutputBuffer, BatchSize, OutputBufferSize, OutputBufferStride);
 			}
 
 			uint32 InputOutputSize = 0;
@@ -1297,366 +1966,1344 @@ namespace UE::NNE::RuntimeBasic
 		{
 			FMemoryCellInstance(const FMemoryCellLayer& InMemoryCellLayer);
 
-			virtual void SetBatchSize(const uint32 BatchSize) override final;
+			virtual void SetMaxBatchSize(const uint32 MaxBatchSize) override final;
+
+			const FMemoryCellLayer& MemoryCellLayer;
+			TSharedPtr<ILayerInstance> RememberInstance;
+			TSharedPtr<ILayerInstance> PassthroughInstance;
+			TSharedPtr<ILayerInstance> MemoryUpdateInstance;
+			TSharedPtr<ILayerInstance> OutputInputUpdateInstance;
+			TSharedPtr<ILayerInstance> OutputMemoryUpdateInstance;
+			TArray<float> RememberGateBuffer;
+			TArray<float> MemoryUpdateBuffer;
+			TArray<float> PassthroughGateBuffer;
+			TArray<float> OutputMemoryUpdateBuffer;
+			TArray<float> OutputInputUpdateBuffer;
+		};
+
+		struct FMemoryCellLayer : public ILayer
+		{
+			virtual TSharedPtr<ILayerInstance> MakeInstance() const { return MakeShared<FMemoryCellInstance>(*this); };
+			virtual ELayerType GetLayerType() const override final { return ELayerType::MemoryCell; }
+			virtual uint32 GetInputSize() const override final { return InputSize + MemorySize; }
+			virtual uint32 GetOutputSize() const override final { return OutputSize + MemorySize; }
+
+			virtual void SerializationSize(uint64& InOutOffset) const override final
+			{
+				Serialization::Size(InOutOffset, InputSize);
+				Serialization::Size(InOutOffset, OutputSize);
+				Serialization::Size(InOutOffset, MemorySize);
+				Serialization::Size(InOutOffset, RememberLayer);
+				Serialization::Size(InOutOffset, PassthroughLayer);
+				Serialization::Size(InOutOffset, MemoryUpdateLayer);
+				Serialization::Size(InOutOffset, OutputInputUpdateLayer);
+				Serialization::Size(InOutOffset, OutputMemoryUpdateLayer);
+			}
+
+			virtual void SerializationLoad(uint64& InOutOffset, TConstArrayView<uint8> Data) override final
+			{
+				Serialization::Load(InOutOffset, InputSize, Data);
+				Serialization::Load(InOutOffset, OutputSize, Data);
+				Serialization::Load(InOutOffset, MemorySize, Data);
+				Serialization::Load(InOutOffset, RememberLayer, Data);
+				Serialization::Load(InOutOffset, PassthroughLayer, Data);
+				Serialization::Load(InOutOffset, MemoryUpdateLayer, Data);
+				Serialization::Load(InOutOffset, OutputInputUpdateLayer, Data);
+				Serialization::Load(InOutOffset, OutputMemoryUpdateLayer, Data);
+			}
+
+			virtual void SerializationSave(uint64& InOutOffset, TArrayView<uint8> Data) const override final
+			{
+				Serialization::Save(InOutOffset, InputSize, Data);
+				Serialization::Save(InOutOffset, OutputSize, Data);
+				Serialization::Save(InOutOffset, MemorySize, Data);
+				Serialization::Save(InOutOffset, RememberLayer, Data);
+				Serialization::Save(InOutOffset, PassthroughLayer, Data);
+				Serialization::Save(InOutOffset, MemoryUpdateLayer, Data);
+				Serialization::Save(InOutOffset, OutputInputUpdateLayer, Data);
+				Serialization::Save(InOutOffset, OutputMemoryUpdateLayer, Data);
+			}
+
+			virtual void Evaluate(
+				ILayerInstance* Instance,
+				float* OutputBuffer,
+				const float* InputBuffer,
+				const uint32 BatchSize,
+				const uint32 OutputBufferSize,
+				const uint32 InputBufferSize,
+				const uint32 OutputBufferStride,
+				const uint32 InputBufferStride) override final
+			{
+				NNE_RUNTIME_BASIC_TRACE_SCOPE(NNE::RuntimeBasic::Private::FMemoryCellLayer::Evaluate);
+				check(OutputBufferSize == GetOutputSize() && InputBufferSize == GetInputSize());
+
+				FMemoryCellInstance* MemoryCellInstance = StaticCast<FMemoryCellInstance*>(Instance);
+				check(MemoryCellInstance);
+				
+				OperatorNanCheck(InputBuffer, BatchSize, InputBufferSize, InputBufferStride);
+
+				// Remember Gate
+
+				RememberLayer->Evaluate(
+					MemoryCellInstance->RememberInstance.Get(),
+					MemoryCellInstance->RememberGateBuffer.GetData(),
+					InputBuffer,
+					BatchSize,
+					MemorySize,
+					InputSize + MemorySize,
+					MemorySize,
+					InputBufferStride);
+
+				// Passthrough Gate
+
+				PassthroughLayer->Evaluate(
+					MemoryCellInstance->PassthroughInstance.Get(),
+					MemoryCellInstance->PassthroughGateBuffer.GetData(),
+					InputBuffer,
+					BatchSize,
+					OutputSize,
+					InputSize + MemorySize,
+					OutputSize,
+					InputBufferStride);
+
+				// Memory Update
+
+				MemoryUpdateLayer->Evaluate(
+					MemoryCellInstance->MemoryUpdateInstance.Get(),
+					MemoryCellInstance->MemoryUpdateBuffer.GetData(),
+					InputBuffer,
+					BatchSize,
+					MemorySize,
+					InputSize + MemorySize,
+					MemorySize,
+					InputBufferStride);
+
+				// Update Memory State
+
+				OperatorMemoryCellUpdateMemory(
+					OutputBuffer + OutputSize,
+					MemoryCellInstance->RememberGateBuffer.GetData(),
+					InputBuffer + InputSize,
+					MemoryCellInstance->MemoryUpdateBuffer.GetData(),
+					BatchSize,
+					MemorySize,
+					OutputBufferStride,
+					MemorySize,
+					InputBufferStride,
+					MemorySize);
+
+				// Output Input Update
+
+				OutputInputUpdateLayer->Evaluate(
+					MemoryCellInstance->OutputInputUpdateInstance.Get(),
+					MemoryCellInstance->OutputInputUpdateBuffer.GetData(),
+					InputBuffer,
+					BatchSize,
+					OutputSize,
+					InputSize + MemorySize,
+					OutputSize,
+					InputBufferStride);
+
+				// Output Memory Update
+
+				OutputMemoryUpdateLayer->Evaluate(
+					MemoryCellInstance->OutputMemoryUpdateInstance.Get(),
+					MemoryCellInstance->OutputMemoryUpdateBuffer.GetData(),
+					OutputBuffer + OutputSize,
+					BatchSize,
+					OutputSize,
+					MemorySize,
+					OutputSize,
+					OutputBufferStride);
+
+				// Update Final Output
+
+				OperatorMemoryCellUpdateOutput(
+					OutputBuffer,
+					MemoryCellInstance->PassthroughGateBuffer.GetData(),
+					MemoryCellInstance->OutputMemoryUpdateBuffer.GetData(),
+					MemoryCellInstance->OutputInputUpdateBuffer.GetData(),
+					BatchSize,
+					OutputSize,
+					OutputBufferStride,
+					OutputSize,
+					OutputSize,
+					OutputSize);
+
+				OperatorNanCheck(OutputBuffer, BatchSize, OutputBufferSize, OutputBufferStride);
+			}
+
+			uint32 InputSize = 0;
+			uint32 OutputSize = 0;
+			uint32 MemorySize = 0;
+			TSharedPtr<ILayer> RememberLayer;
+			TSharedPtr<ILayer> PassthroughLayer;
+			TSharedPtr<ILayer> MemoryUpdateLayer;
+			TSharedPtr<ILayer> OutputInputUpdateLayer;
+			TSharedPtr<ILayer> OutputMemoryUpdateLayer;
+		};
+
+		FMemoryCellInstance::FMemoryCellInstance(const FMemoryCellLayer& InMemoryCellLayer) 
+			: MemoryCellLayer(InMemoryCellLayer)
+		{
+			RememberInstance = MemoryCellLayer.RememberLayer->MakeInstance();
+			PassthroughInstance = MemoryCellLayer.RememberLayer->MakeInstance();
+			MemoryUpdateInstance = MemoryCellLayer.RememberLayer->MakeInstance();
+			OutputInputUpdateInstance = MemoryCellLayer.RememberLayer->MakeInstance();
+			OutputMemoryUpdateInstance = MemoryCellLayer.RememberLayer->MakeInstance();
+		}
+
+		void FMemoryCellInstance::SetMaxBatchSize(const uint32 MaxBatchSize)
+		{
+			NNE_RUNTIME_BASIC_TRACE_SCOPE(NNE::RuntimeBasic::Private::FMemoryCellInstance::SetMaxBatchSize);
+
+			if (RememberInstance) { RememberInstance->SetMaxBatchSize(MaxBatchSize); }
+			if (PassthroughInstance) { PassthroughInstance->SetMaxBatchSize(MaxBatchSize); }
+			if (MemoryUpdateInstance) { MemoryUpdateInstance->SetMaxBatchSize(MaxBatchSize); }
+			if (OutputInputUpdateInstance) { OutputInputUpdateInstance->SetMaxBatchSize(MaxBatchSize); }
+			if (OutputMemoryUpdateInstance) { OutputMemoryUpdateInstance->SetMaxBatchSize(MaxBatchSize); }
+
+			RememberGateBuffer.SetNumUninitialized(MaxBatchSize * MemoryCellLayer.MemorySize, false);
+			PassthroughGateBuffer.SetNumUninitialized(MaxBatchSize * MemoryCellLayer.OutputSize, false);
+			MemoryUpdateBuffer.SetNumUninitialized(MaxBatchSize * MemoryCellLayer.MemorySize, false);
+			OutputInputUpdateBuffer.SetNumUninitialized(MaxBatchSize * MemoryCellLayer.OutputSize, false);
+			OutputMemoryUpdateBuffer.SetNumUninitialized(MaxBatchSize * MemoryCellLayer.OutputSize, false);
+		}
+
+		//--------------------------------------------------------------------------
+
+		struct FCopyLayer : public ILayer
+		{
+			virtual ELayerType GetLayerType() const override final { return ELayerType::Copy; }
+			virtual uint32 GetInputSize() const override final { return InputOutputSize; }
+			virtual uint32 GetOutputSize() const override final { return InputOutputSize; }
+
+			virtual void SerializationSize(uint64& InOutOffset) const override final
+			{
+				Serialization::Size(InOutOffset, InputOutputSize);
+			}
+
+			virtual void SerializationLoad(uint64& InOutOffset, TConstArrayView<uint8> Data) override final
+			{
+				Serialization::Load(InOutOffset, InputOutputSize, Data);
+			}
+
+			virtual void SerializationSave(uint64& InOutOffset, TArrayView<uint8> Data) const override final
+			{
+				Serialization::Save(InOutOffset, InputOutputSize, Data);
+			}
+
+			virtual void Evaluate(
+				ILayerInstance* Instance,
+				float* OutputBuffer,
+				const float* InputBuffer,
+				const uint32 BatchSize,
+				const uint32 OutputBufferSize,
+				const uint32 InputBufferSize,
+				const uint32 OutputBufferStride,
+				const uint32 InputBufferStride) override final
+			{
+				NNE_RUNTIME_BASIC_TRACE_SCOPE(NNE::RuntimeBasic::Private::FCopyLayer::Evaluate);
+				check(OutputBufferSize == GetOutputSize() && InputBufferSize == GetInputSize());
+				check(Instance == nullptr);
+				OperatorNanCheck(InputBuffer, BatchSize, InputBufferSize, InputBufferStride);
+
+				OperatorCopy(
+					OutputBuffer,
+					InputBuffer,
+					BatchSize,
+					InputOutputSize,
+					OutputBufferStride,
+					InputBufferStride);
+
+				OperatorNanCheck(OutputBuffer, BatchSize, OutputBufferSize, OutputBufferStride);
+			}
+
+			uint32 InputOutputSize = 0;
+		};
+
+		//--------------------------------------------------------------------------
+
+		struct FConcatLayer;
+
+		struct FConcatLayerInstance : public ILayerInstance
+		{
+			FConcatLayerInstance(const FConcatLayer& InConcatLayer);
+
+			void SetMaxBatchSize(const uint32 MaxBatchSize) override final;
+
+			const FConcatLayer& ConcatLayer;
+			TArray<TSharedPtr<ILayerInstance>, TInlineAllocator<32>> Instances;
+		};
+
+		struct FConcatLayer : public ILayer
+		{
+			virtual TSharedPtr<ILayerInstance> MakeInstance() const { return MakeShared<FConcatLayerInstance>(*this); };
+			virtual ELayerType GetLayerType() const override final { return ELayerType::Concat; }
+			virtual uint32 GetInputSize() const override final { return TotalInputSize; }
+			virtual uint32 GetOutputSize() const override final { return TotalOutputSize; }
+
+			virtual void SerializationSize(uint64& InOutOffset) const override final
+			{
+				Serialization::Size(InOutOffset, (uint32)Layers.Num());
+				Serialization::Size(InOutOffset, InputSizes);
+				Serialization::Size(InOutOffset, OutputSizes);
+				Serialization::Size(InOutOffset, Layers);
+			}
+
+			virtual void SerializationLoad(uint64& InOutOffset, TConstArrayView<uint8> Data) override final
+			{
+				uint32 LayerNum = 0;
+				Serialization::Load(InOutOffset, LayerNum, Data);
+				Serialization::Load(InOutOffset, InputSizes, Data, LayerNum);
+				Serialization::Load(InOutOffset, OutputSizes, Data, LayerNum);
+				Layers.Init(nullptr, LayerNum);
+				Serialization::Load(InOutOffset, Layers, Data);
+				PostLoad();
+			}
+
+			void PostLoad()
+			{
+				const uint32 LayerNum = Layers.Num();
+				InputOffsets.SetNumUninitialized(LayerNum);
+				OutputOffsets.SetNumUninitialized(LayerNum);
+
+				TotalInputSize = 0;
+				TotalOutputSize = 0;
+				for (uint32 LayerIdx = 0; LayerIdx < LayerNum; LayerIdx++)
+				{
+					InputOffsets[LayerIdx] = TotalInputSize;
+					OutputOffsets[LayerIdx] = TotalOutputSize;
+					TotalInputSize += InputSizes[LayerIdx];
+					TotalOutputSize += OutputSizes[LayerIdx];
+				}
+			}
+
+			virtual void SerializationSave(uint64& InOutOffset, TArrayView<uint8> Data) const override final
+			{
+				Serialization::Save(InOutOffset, (uint32)Layers.Num(), Data);
+				Serialization::Save(InOutOffset, InputSizes, Data);
+				Serialization::Save(InOutOffset, OutputSizes, Data);
+				Serialization::Save(InOutOffset, Layers, Data);
+			}
+
+			virtual void Evaluate(
+				ILayerInstance* Instance,
+				float* OutputBuffer,
+				const float* InputBuffer,
+				const uint32 BatchSize,
+				const uint32 OutputBufferSize,
+				const uint32 InputBufferSize,
+				const uint32 OutputBufferStride,
+				const uint32 InputBufferStride) override final
+			{
+				NNE_RUNTIME_BASIC_TRACE_SCOPE(NNE::RuntimeBasic::Private::FConcatLayer::Evaluate);
+				check(OutputBufferSize == GetOutputSize() && InputBufferSize == GetInputSize());
+
+				FConcatLayerInstance* ConcatInstance = StaticCast<FConcatLayerInstance*>(Instance);
+				check(ConcatInstance);
+
+				OperatorNanCheck(InputBuffer, BatchSize, InputBufferSize, InputBufferStride);
+
+				const int32 LayerNum = Layers.Num();
+
+				for (int32 LayerIdx = 0; LayerIdx < LayerNum; LayerIdx++)
+				{
+					Layers[LayerIdx]->Evaluate(
+						ConcatInstance->Instances[LayerIdx].Get(),
+						OutputBuffer + OutputOffsets[LayerIdx],
+						InputBuffer + InputOffsets[LayerIdx],
+						BatchSize,
+						OutputSizes[LayerIdx],
+						InputSizes[LayerIdx],
+						OutputBufferStride,
+						InputBufferStride);
+				}
+
+				OperatorNanCheck(OutputBuffer, BatchSize, OutputBufferSize, OutputBufferStride);
+			}
+
+			TConstArrayView<uint32> InputSizes;
+			TConstArrayView<uint32> OutputSizes;
+			TArray<TSharedPtr<ILayer>, TInlineAllocator<32>> Layers;
+
+			uint32 TotalInputSize = 0;
+			uint32 TotalOutputSize = 0;
+			TArray<uint32, TInlineAllocator<32>> InputOffsets;
+			TArray<uint32, TInlineAllocator<32>> OutputOffsets;
+		};
+
+		FConcatLayerInstance::FConcatLayerInstance(const FConcatLayer& InConcatLayer)
+			: ConcatLayer(InConcatLayer)
+		{
+			const uint32 LayerNum = ConcatLayer.Layers.Num();
+			Instances.Init(nullptr, LayerNum);
+
+			for (uint32 LayerIdx = 0; LayerIdx < LayerNum; LayerIdx++)
+			{
+				Instances[LayerIdx] = ConcatLayer.Layers[LayerIdx]->MakeInstance();
+			}
+		}
+
+		void FConcatLayerInstance::SetMaxBatchSize(const uint32 MaxBatchSize)
+		{
+			NNE_RUNTIME_BASIC_TRACE_SCOPE(NNE::RuntimeBasic::Private::FConcatLayerInstance::SetMaxBatchSize);
+
+			// Propagate call to sub-layer instances
+
+			for (const TSharedPtr<ILayerInstance>& Instance : Instances)
+			{
+				if (Instance)
+				{
+					Instance->SetMaxBatchSize(MaxBatchSize);
+				}
+			}
+		}
+
+		//--------------------------------------------------------------------------
+
+		struct FArrayLayer;
+
+		struct FArrayLayerInstance : public ILayerInstance
+		{
+			FArrayLayerInstance(const FArrayLayer& InArrayLayer);
+
+			void SetMaxBatchSize(const uint32 MaxBatchSize) override final;
+
+			const FArrayLayer& ArrayLayer;
+			TSharedPtr<ILayerInstance> Instance;
+			TArray<float> ElementInputBuffer;
+			TArray<float> ElementOutputBuffer;
+		};
+
+		struct FArrayLayer : public ILayer
+		{
+			virtual TSharedPtr<ILayerInstance> MakeInstance() const { return MakeShared<FArrayLayerInstance>(*this); };
+			virtual ELayerType GetLayerType() const override final { return ELayerType::Array; }
+			virtual uint32 GetInputSize() const override final { return ElementNum * ElementInputSize; }
+			virtual uint32 GetOutputSize() const override final { return ElementNum * ElementOutputSize; }
+
+			virtual void SerializationSize(uint64& InOutOffset) const override final
+			{
+				Serialization::Size(InOutOffset, ElementNum);
+				Serialization::Size(InOutOffset, ElementInputSize);
+				Serialization::Size(InOutOffset, ElementOutputSize);
+				Serialization::Size(InOutOffset, SubLayer);
+			}
+
+			virtual void SerializationLoad(uint64& InOutOffset, TConstArrayView<uint8> Data) override final
+			{
+				Serialization::Load(InOutOffset, ElementNum, Data);
+				Serialization::Load(InOutOffset, ElementInputSize, Data);
+				Serialization::Load(InOutOffset, ElementOutputSize, Data);
+				Serialization::Load(InOutOffset, SubLayer, Data);
+			}
+
+			virtual void SerializationSave(uint64& InOutOffset, TArrayView<uint8> Data) const override final
+			{
+				Serialization::Save(InOutOffset, ElementNum, Data);
+				Serialization::Save(InOutOffset, ElementInputSize, Data);
+				Serialization::Save(InOutOffset, ElementOutputSize, Data);
+				Serialization::Save(InOutOffset, SubLayer, Data);
+			}
+
+			virtual void Evaluate(
+				ILayerInstance* Instance,
+				float* OutputBuffer,
+				const float* InputBuffer,
+				const uint32 BatchSize,
+				const uint32 OutputBufferSize,
+				const uint32 InputBufferSize,
+				const uint32 OutputBufferStride,
+				const uint32 InputBufferStride) override final
+			{
+				NNE_RUNTIME_BASIC_TRACE_SCOPE(NNE::RuntimeBasic::Private::FArrayLayer::Evaluate);
+				check(OutputBufferSize == GetOutputSize() && InputBufferSize == GetInputSize());
+
+				FArrayLayerInstance* ArrayInstance = StaticCast<FArrayLayerInstance*>(Instance);
+				check(ArrayInstance);
+
+				OperatorNanCheck(InputBuffer, BatchSize, InputBufferSize, InputBufferStride);
+
+				// If inputs and outputs are already tightly packed then evaluate directly as large batch
+
+				if (InputBufferStride == ElementNum * ElementInputSize &&
+					OutputBufferStride == ElementNum * ElementOutputSize)
+				{
+					SubLayer->Evaluate(
+						ArrayInstance->Instance.Get(),
+						OutputBuffer,
+						InputBuffer,
+						BatchSize * ElementNum,
+						ElementOutputSize,
+						ElementInputSize,
+						ElementOutputSize,
+						ElementInputSize);
+
+					return;
+				}
+
+				// Otherwise gather all inputs into one large buffer packed together tightly
+
+				OperatorCopy(
+					ArrayInstance->ElementInputBuffer.GetData(),
+					InputBuffer,
+					BatchSize,
+					ElementNum * ElementInputSize,
+					ElementNum * ElementInputSize,
+					InputBufferStride);
+
+				// Evaluate sub-layer on large batch of all elements
+
+				SubLayer->Evaluate(
+					ArrayInstance->Instance.Get(),
+					ArrayInstance->ElementOutputBuffer.GetData(),
+					ArrayInstance->ElementInputBuffer.GetData(),
+					BatchSize * ElementNum,
+					ElementOutputSize,
+					ElementInputSize,
+					ElementOutputSize,
+					ElementInputSize);
+
+				// And scatter outputs out of tightly packed buffer
+
+				OperatorCopy(
+					OutputBuffer,
+					ArrayInstance->ElementOutputBuffer.GetData(),
+					BatchSize,
+					ElementNum * ElementOutputSize,
+					OutputBufferStride,
+					ElementNum * ElementOutputSize);
+
+				OperatorNanCheck(OutputBuffer, BatchSize, OutputBufferSize, OutputBufferStride);
+			}
+
+			uint32 ElementNum = 0;
+			uint32 ElementInputSize = 0;
+			uint32 ElementOutputSize = 0;
+			TSharedPtr<ILayer> SubLayer;
+		};
+
+		FArrayLayerInstance::FArrayLayerInstance(const FArrayLayer& InArrayLayer) : ArrayLayer(InArrayLayer)
+		{
+			Instance = ArrayLayer.SubLayer->MakeInstance();
+		}
+
+		void FArrayLayerInstance::SetMaxBatchSize(const uint32 MaxBatchSize)
+		{
+			NNE_RUNTIME_BASIC_TRACE_SCOPE(NNE::RuntimeBasic::Private::FArrayLayerInstance::SetMaxBatchSize);
+
+			if (Instance)
+			{
+				// We are going to evaluate the sublayer on one large batch so we use MaxBatchSize * ArrayLayer.ElementNum.
+				Instance->SetMaxBatchSize(MaxBatchSize * ArrayLayer.ElementNum);
+			}
+
+			ElementInputBuffer.SetNumUninitialized(MaxBatchSize * ArrayLayer.ElementNum * ArrayLayer.ElementInputSize, false);
+			ElementOutputBuffer.SetNumUninitialized(MaxBatchSize * ArrayLayer.ElementNum * ArrayLayer.ElementOutputSize, false);
+		}
+
+		//--------------------------------------------------------------------------
+
+		struct FAggregateSetLayer;
+
+		struct FAggregateSetLayerInstance : public ILayerInstance
+		{
+			FAggregateSetLayerInstance(const FAggregateSetLayer& InAggregateSetLayer);
+
+			void SetMaxBatchSize(const uint32 MaxBatchSize) override final;
+
+			const FAggregateSetLayer& AggregateSetLayer;
+			TSharedPtr<ILayerInstance> SubLayerInstance;
+			TSharedPtr<ILayerInstance> QueryInstance;
+			TSharedPtr<ILayerInstance> KeyInstance;
+			TSharedPtr<ILayerInstance> ValueInstance;
+
+			uint32 TotalElementNum = 0;
+			TArray<uint32, TInlineAllocator<32>> ElementNums;
+			TArray<uint32, TInlineAllocator<32>> ElementOffsets;
+
+			TArray<float> InputElementBuffer;
+			TArray<float> OutputElementBuffer;
+			TArray<float> QueryBuffer;
+			TArray<float> KeyBuffer;
+			TArray<float> ValueBuffer;
+			TArray<float> AttentionMaxsBuffer;
+			TArray<float> AttentionDenomsBuffer;
+			TArray<float> AttentionBuffer;
+		};
+
+		struct FAggregateSetLayer : public ILayer
+		{
+			virtual TSharedPtr<ILayerInstance> MakeInstance() const { return MakeShared<FAggregateSetLayerInstance>(*this); };
+			virtual ELayerType GetLayerType() const override final { return ELayerType::AggregateSet; }
+			virtual uint32 GetInputSize() const override final { return MaxElementNum * ElementInputSize + MaxElementNum; }
+			virtual uint32 GetOutputSize() const override final { return AttentionHeadNum * OutputEncodingSize + MaxElementNum; }
+
+			virtual void SerializationSize(uint64& InOutOffset) const override final
+			{
+				Serialization::Size(InOutOffset, MaxElementNum);
+				Serialization::Size(InOutOffset, ElementInputSize);
+				Serialization::Size(InOutOffset, ElementOutputSize);
+				Serialization::Size(InOutOffset, OutputEncodingSize);
+				Serialization::Size(InOutOffset, AttentionEncodingSize);
+				Serialization::Size(InOutOffset, AttentionHeadNum);
+				Serialization::Size(InOutOffset, SubLayer);
+				Serialization::Size(InOutOffset, QueryLayer);
+				Serialization::Size(InOutOffset, KeyLayer);
+				Serialization::Size(InOutOffset, ValueLayer);
+			}
+
+			virtual void SerializationLoad(uint64& InOutOffset, TConstArrayView<uint8> Data) override final
+			{
+				Serialization::Load(InOutOffset, MaxElementNum, Data);
+				Serialization::Load(InOutOffset, ElementInputSize, Data);
+				Serialization::Load(InOutOffset, ElementOutputSize, Data);
+				Serialization::Load(InOutOffset, OutputEncodingSize, Data);
+				Serialization::Load(InOutOffset, AttentionEncodingSize, Data);
+				Serialization::Load(InOutOffset, AttentionHeadNum, Data);
+				Serialization::Load(InOutOffset, SubLayer, Data);
+				Serialization::Load(InOutOffset, QueryLayer, Data);
+				Serialization::Load(InOutOffset, KeyLayer, Data);
+				Serialization::Load(InOutOffset, ValueLayer, Data);
+			}
+
+			virtual void SerializationSave(uint64& InOutOffset, TArrayView<uint8> Data) const override final
+			{
+				Serialization::Save(InOutOffset, MaxElementNum, Data);
+				Serialization::Save(InOutOffset, ElementInputSize, Data);
+				Serialization::Save(InOutOffset, ElementOutputSize, Data);
+				Serialization::Save(InOutOffset, OutputEncodingSize, Data);
+				Serialization::Save(InOutOffset, AttentionEncodingSize, Data);
+				Serialization::Save(InOutOffset, AttentionHeadNum, Data);
+				Serialization::Save(InOutOffset, SubLayer, Data);
+				Serialization::Save(InOutOffset, QueryLayer, Data);
+				Serialization::Save(InOutOffset, KeyLayer, Data);
+				Serialization::Save(InOutOffset, ValueLayer, Data);
+			}
+
+			virtual void Evaluate(
+				ILayerInstance* Instance,
+				float* OutputBuffer,
+				const float* InputBuffer,
+				const uint32 BatchSize,
+				const uint32 OutputBufferSize,
+				const uint32 InputBufferSize,
+				const uint32 OutputBufferStride,
+				const uint32 InputBufferStride) override final
+			{
+				NNE_RUNTIME_BASIC_TRACE_SCOPE(NNE::RuntimeBasic::Private::FAggregateSetLayer::Evaluate);
+				check(OutputBufferSize == GetOutputSize() && InputBufferSize == GetInputSize());
+
+				FAggregateSetLayerInstance* AggregateSetInstance = StaticCast<FAggregateSetLayerInstance*>(Instance);
+				check(AggregateSetInstance);
+
+				OperatorNanCheck(InputBuffer, BatchSize, InputBufferSize, InputBufferStride);
+
+				// Count the number of elements for each item in the batch
+				
+				OperatorAggregateCountElementNum(
+					AggregateSetInstance->TotalElementNum,
+					AggregateSetInstance->ElementNums.GetData(),
+					AggregateSetInstance->ElementOffsets.GetData(),
+					InputBuffer + MaxElementNum * ElementInputSize,
+					BatchSize,
+					MaxElementNum,
+					InputBufferStride);
+
+				// Gather Elements from all batches into one large tightly packed buffer
+
+				OperatorAggregateGatherElements(
+					AggregateSetInstance->InputElementBuffer.GetData(),
+					InputBuffer,
+					AggregateSetInstance->ElementNums.GetData(),
+					AggregateSetInstance->ElementOffsets.GetData(),
+					BatchSize,
+					ElementInputSize,
+					InputBufferStride);
+
+				// Evaluate Sublayer on all elements
+
+				SubLayer->Evaluate(
+					AggregateSetInstance->SubLayerInstance.Get(),
+					AggregateSetInstance->OutputElementBuffer.GetData(),
+					AggregateSetInstance->InputElementBuffer.GetData(),
+					AggregateSetInstance->TotalElementNum,
+					ElementOutputSize,
+					ElementInputSize,
+					ElementOutputSize,
+					ElementInputSize);
+
+				// Compute Query on all elements
+
+				QueryLayer->Evaluate(
+					AggregateSetInstance->QueryInstance.Get(),
+					AggregateSetInstance->QueryBuffer.GetData(),
+					AggregateSetInstance->OutputElementBuffer.GetData(),
+					AggregateSetInstance->TotalElementNum,
+					AttentionHeadNum * AttentionEncodingSize,
+					ElementOutputSize,
+					AttentionHeadNum * AttentionEncodingSize,
+					ElementOutputSize);
+
+				// Compute Keys on all elements
+
+				KeyLayer->Evaluate(
+					AggregateSetInstance->KeyInstance.Get(),
+					AggregateSetInstance->KeyBuffer.GetData(),
+					AggregateSetInstance->OutputElementBuffer.GetData(),
+					AggregateSetInstance->TotalElementNum,
+					AttentionHeadNum * AttentionEncodingSize,
+					ElementOutputSize,
+					AttentionHeadNum * AttentionEncodingSize,
+					ElementOutputSize);
+
+				// Compute Values on all elements
+
+				ValueLayer->Evaluate(
+					AggregateSetInstance->ValueInstance.Get(),
+					AggregateSetInstance->ValueBuffer.GetData(),
+					AggregateSetInstance->OutputElementBuffer.GetData(),
+					AggregateSetInstance->TotalElementNum,
+					AttentionHeadNum * OutputEncodingSize,
+					ElementOutputSize,
+					AttentionHeadNum * OutputEncodingSize,
+					ElementOutputSize);
+
+				// Compute Attention
+
+				OperatorAggregateDotProductAttention(
+					AggregateSetInstance->AttentionBuffer.GetData(),
+					AggregateSetInstance->QueryBuffer.GetData(),
+					AggregateSetInstance->KeyBuffer.GetData(),
+					AggregateSetInstance->TotalElementNum,
+					AttentionEncodingSize,
+					AttentionHeadNum);
+
+				OperatorAggregateSoftmaxPlusOneInplace(
+					AggregateSetInstance->AttentionMaxsBuffer.GetData(),
+					AggregateSetInstance->AttentionDenomsBuffer.GetData(),
+					AggregateSetInstance->AttentionBuffer.GetData(),
+					AggregateSetInstance->ElementNums.GetData(),
+					AggregateSetInstance->ElementOffsets.GetData(),
+					BatchSize,
+					AttentionHeadNum);
+
+				OperatorAggregateAttentionSum(
+					OutputBuffer,
+					AggregateSetInstance->AttentionBuffer.GetData(),
+					AggregateSetInstance->ValueBuffer.GetData(),
+					AggregateSetInstance->ElementNums.GetData(),
+					AggregateSetInstance->ElementOffsets.GetData(),
+					BatchSize,
+					OutputEncodingSize,
+					AttentionHeadNum,
+					OutputBufferStride);
+
+				// Append Element Mask
+
+				OperatorCopy(
+					OutputBuffer + AttentionHeadNum * OutputEncodingSize,
+					InputBuffer + MaxElementNum * ElementInputSize,
+					BatchSize,
+					MaxElementNum,
+					OutputBufferStride,
+					InputBufferStride);
+
+				OperatorNanCheck(OutputBuffer, BatchSize, OutputBufferSize, OutputBufferStride);
+			}
+
+			uint32 MaxElementNum = 0;
+			uint32 OutputEncodingSize = 0;
+			uint32 AttentionEncodingSize = 0;
+			uint32 AttentionHeadNum = 0;
+
+			TSharedPtr<ILayer> SubLayer;
+			TSharedPtr<ILayer> QueryLayer;
+			TSharedPtr<ILayer> KeyLayer;
+			TSharedPtr<ILayer> ValueLayer;
+
+			uint32 ElementInputSize = 0;
+			uint32 ElementOutputSize = 0;
+		};
+
+		FAggregateSetLayerInstance::FAggregateSetLayerInstance(const FAggregateSetLayer& InAggregateSetLayer)
+			: AggregateSetLayer(InAggregateSetLayer)
+		{
+			SubLayerInstance = AggregateSetLayer.SubLayer->MakeInstance();
+			QueryInstance = AggregateSetLayer.QueryLayer->MakeInstance();
+			KeyInstance = AggregateSetLayer.KeyLayer->MakeInstance();
+			ValueInstance = AggregateSetLayer.ValueLayer->MakeInstance();
+		}
+
+		void FAggregateSetLayerInstance::SetMaxBatchSize(const uint32 MaxBatchSize)
+		{
+			NNE_RUNTIME_BASIC_TRACE_SCOPE(NNE::RuntimeBasic::Private::FAggregateSetLayerInstance::SetMaxBatchSize);
+
+			if (SubLayerInstance) { SubLayerInstance->SetMaxBatchSize(MaxBatchSize * AggregateSetLayer.MaxElementNum); }
+			if (QueryInstance) { QueryInstance->SetMaxBatchSize(MaxBatchSize * AggregateSetLayer.MaxElementNum); }
+			if (KeyInstance) { KeyInstance->SetMaxBatchSize(MaxBatchSize * AggregateSetLayer.MaxElementNum); }
+			if (ValueInstance) { ValueInstance->SetMaxBatchSize(MaxBatchSize * AggregateSetLayer.MaxElementNum); }
+
+			ElementNums.SetNumUninitialized(MaxBatchSize, false);
+			ElementOffsets.SetNumUninitialized(MaxBatchSize, false);
+			
+			InputElementBuffer.SetNumUninitialized(MaxBatchSize * AggregateSetLayer.MaxElementNum * AggregateSetLayer.ElementInputSize, false);
+			OutputElementBuffer.SetNumUninitialized(MaxBatchSize * AggregateSetLayer.MaxElementNum * AggregateSetLayer.ElementOutputSize, false);
+			QueryBuffer.SetNumUninitialized(MaxBatchSize * AggregateSetLayer.MaxElementNum * AggregateSetLayer.AttentionHeadNum * AggregateSetLayer.AttentionEncodingSize, false);
+			KeyBuffer.SetNumUninitialized(MaxBatchSize * AggregateSetLayer.MaxElementNum * AggregateSetLayer.AttentionHeadNum * AggregateSetLayer.AttentionEncodingSize, false);
+			ValueBuffer.SetNumUninitialized(MaxBatchSize * AggregateSetLayer.MaxElementNum * AggregateSetLayer.AttentionHeadNum * AggregateSetLayer.OutputEncodingSize, false);
+			AttentionMaxsBuffer.SetNumUninitialized(MaxBatchSize * AggregateSetLayer.AttentionHeadNum, false);
+			AttentionDenomsBuffer.SetNumUninitialized(MaxBatchSize * AggregateSetLayer.AttentionHeadNum, false);
+			AttentionBuffer.SetNumUninitialized(MaxBatchSize * AggregateSetLayer.MaxElementNum * AggregateSetLayer.AttentionHeadNum, false);
+		}
+
+		//--------------------------------------------------------------------------
+
+		struct FAggregateOrExclusiveLayer;
+
+		struct FAggregateOrExclusiveLayerInstance : public ILayerInstance
+		{
+			FAggregateOrExclusiveLayerInstance(const FAggregateOrExclusiveLayer& InAggregateOrExclusiveLayer);
+
+			void SetMaxBatchSize(const uint32 MaxBatchSize) override final;
+
+			const FAggregateOrExclusiveLayer& AggregateOrExclusiveLayer;
+			TArray<TSharedPtr<ILayerInstance>, TInlineAllocator<32>> SubLayerInstances;
+			TArray<TSharedPtr<ILayerInstance>, TInlineAllocator<32>> EncoderInstances;
 
-			const FMemoryCellLayer& MemoryCellLayer;
-			TArray<float, TInlineAllocator<512>> RememberGateBuffer;
-			TArray<float, TInlineAllocator<512>> MemoryUpdateBuffer;
-			TArray<float, TInlineAllocator<512>> PassthroughGateBuffer;
-			TArray<float, TInlineAllocator<512>> OutputMemoryUpdateBuffer;
-			TArray<float, TInlineAllocator<512>> OutputInputUpdateBuffer;
+			TArray<TArray<uint32>, TInlineAllocator<32>> SubLayerBatchIndices;
+			TArray<TArray<float>, TInlineAllocator<32>> SubLayerInputBuffers;
+			TArray<TArray<float>, TInlineAllocator<32>> SubLayerOutputBuffers;
+			TArray<TArray<float>, TInlineAllocator<32>> SubLayerEncodingBuffers;
 		};
 
-		//--------------------------------------------------------------------------
-
-		struct FMemoryCellLayer : public ILayer
+		struct FAggregateOrExclusiveLayer : public ILayer
 		{
-			virtual TSharedPtr<ILayerInstance> MakeInstance() const { return MakeShared<FMemoryCellInstance>(*this); };
-			virtual ELayerType GetLayerType() const override final { return ELayerType::MemoryCell; }
-			virtual uint32 GetInputSize() const override final { return InputSize + MemorySize; }
-			virtual uint32 GetOutputSize() const override final { return OutputSize + MemorySize; }
+			virtual TSharedPtr<ILayerInstance> MakeInstance() const { return MakeShared<FAggregateOrExclusiveLayerInstance>(*this); };
+			virtual ELayerType GetLayerType() const override final { return ELayerType::AggregateOrExclusive; }
+			virtual uint32 GetInputSize() const override final { return MaxSubLayerInputSize + SubLayers.Num(); }
+			virtual uint32 GetOutputSize() const override final { return OutputEncodingSize + SubLayers.Num(); }
 
 			virtual void SerializationSize(uint64& InOutOffset) const override final
 			{
-				Serialization::Size(InOutOffset, InputSize);
-				Serialization::Size(InOutOffset, OutputSize);
-				Serialization::Size(InOutOffset, MemorySize);
-				Serialization::Size(InOutOffset, BiasesWr);
-				Serialization::Size(InOutOffset, WeightsWr);
-				Serialization::Size(InOutOffset, BiasesWp);
-				Serialization::Size(InOutOffset, WeightsWp);
-				Serialization::Size(InOutOffset, BiasesWn);
-				Serialization::Size(InOutOffset, WeightsWn);
-				Serialization::Size(InOutOffset, BiasesWz);
-				Serialization::Size(InOutOffset, WeightsWz);
-				Serialization::Size(InOutOffset, BiasesWy);
-				Serialization::Size(InOutOffset, WeightsWy);
+				Serialization::Size(InOutOffset, (uint32)SubLayers.Num());
+				Serialization::Size(InOutOffset, OutputEncodingSize);
+				Serialization::Size(InOutOffset, SubLayerInputSizes);
+				Serialization::Size(InOutOffset, SubLayerOutputSizes);
+				Serialization::Size(InOutOffset, SubLayers);
+				Serialization::Size(InOutOffset, Encoders);
 			}
 
 			virtual void SerializationLoad(uint64& InOutOffset, TConstArrayView<uint8> Data) override final
 			{
-				Serialization::Load(InOutOffset, InputSize, Data);
-				Serialization::Load(InOutOffset, OutputSize, Data);
-				Serialization::Load(InOutOffset, MemorySize, Data);
-				Serialization::Load(InOutOffset, BiasesWr, Data, MemorySize);
-				Serialization::Load(InOutOffset, WeightsWr, Data, (InputSize + MemorySize) * MemorySize);
-				Serialization::Load(InOutOffset, BiasesWp, Data, OutputSize);
-				Serialization::Load(InOutOffset, WeightsWp, Data, (InputSize + MemorySize) * OutputSize);
-				Serialization::Load(InOutOffset, BiasesWn, Data, MemorySize);
-				Serialization::Load(InOutOffset, WeightsWn, Data, (InputSize + MemorySize) * MemorySize);
-				Serialization::Load(InOutOffset, BiasesWz, Data, OutputSize);
-				Serialization::Load(InOutOffset, WeightsWz, Data, (InputSize + MemorySize) * OutputSize);
-				Serialization::Load(InOutOffset, BiasesWy, Data, OutputSize);
-				Serialization::Load(InOutOffset, WeightsWy, Data, MemorySize * OutputSize);
+				uint32 SubLayerNum = 0;
+				Serialization::Load(InOutOffset, SubLayerNum, Data);
+				Serialization::Load(InOutOffset, OutputEncodingSize, Data);
+				Serialization::Load(InOutOffset, SubLayerInputSizes, Data, SubLayerNum);
+				Serialization::Load(InOutOffset, SubLayerOutputSizes, Data, SubLayerNum);
+				SubLayers.Init(nullptr, SubLayerNum);
+				Serialization::Load(InOutOffset, SubLayers, Data);
+				Encoders.Init(nullptr, SubLayerNum);
+				Serialization::Load(InOutOffset, Encoders, Data);
+				PostLoad();
+			}
+
+			void PostLoad()
+			{
+				const uint32 SubLayerNum = SubLayers.Num();
+				MaxSubLayerInputSize = 0;
+				for (uint32 SubLayerIdx = 0; SubLayerIdx < SubLayerNum; SubLayerIdx++)
+				{
+					MaxSubLayerInputSize = FMath::Max(MaxSubLayerInputSize, SubLayerInputSizes[SubLayerIdx]);
+				}
 			}
 
 			virtual void SerializationSave(uint64& InOutOffset, TArrayView<uint8> Data) const override final
 			{
-				Serialization::Save(InOutOffset, InputSize, Data);
-				Serialization::Save(InOutOffset, OutputSize, Data);
-				Serialization::Save(InOutOffset, MemorySize, Data);
-				Serialization::Save(InOutOffset, BiasesWr, Data);
-				Serialization::Save(InOutOffset, WeightsWr, Data);
-				Serialization::Save(InOutOffset, BiasesWp, Data);
-				Serialization::Save(InOutOffset, WeightsWp, Data);
-				Serialization::Save(InOutOffset, BiasesWn, Data);
-				Serialization::Save(InOutOffset, WeightsWn, Data);
-				Serialization::Save(InOutOffset, BiasesWz, Data);
-				Serialization::Save(InOutOffset, WeightsWz, Data);
-				Serialization::Save(InOutOffset, BiasesWy, Data);
-				Serialization::Save(InOutOffset, WeightsWy, Data);
+				Serialization::Save(InOutOffset, (uint32)SubLayers.Num(), Data);
+				Serialization::Save(InOutOffset, OutputEncodingSize, Data);
+				Serialization::Save(InOutOffset, SubLayerInputSizes, Data);
+				Serialization::Save(InOutOffset, SubLayerOutputSizes, Data);
+				Serialization::Save(InOutOffset, SubLayers, Data);
+				Serialization::Save(InOutOffset, Encoders, Data);
 			}
 
 			virtual void Evaluate(
-				TArrayView<float> Output,
-				TConstArrayView<float> Input,
 				ILayerInstance* Instance,
-				uint32 BatchSize,
-				uint32 OutputStride,
-				uint32 InputStride) override final
+				float* OutputBuffer,
+				const float* InputBuffer,
+				const uint32 BatchSize,
+				const uint32 OutputBufferSize,
+				const uint32 InputBufferSize,
+				const uint32 OutputBufferStride,
+				const uint32 InputBufferStride) override final
 			{
-				TRACE_CPUPROFILER_EVENT_SCOPE(NNE::RuntimeBasic::Private::FMemoryCellLayer::Evaluate);
+				NNE_RUNTIME_BASIC_TRACE_SCOPE(NNE::RuntimeBasic::Private::FAggregateOrExclusiveLayer::Evaluate);
+				check(OutputBufferSize == GetOutputSize() && InputBufferSize == GetInputSize());
 
-				FMemoryCellInstance* MemoryCellInstance = StaticCast<FMemoryCellInstance*>(Instance);
-				check(MemoryCellInstance);
+				FAggregateOrExclusiveLayerInstance* AggregateOrExclusiveInstance = StaticCast<FAggregateOrExclusiveLayerInstance*>(Instance);
+				check(AggregateOrExclusiveInstance);
 
-				// Remember Gate
+				OperatorNanCheck(InputBuffer, BatchSize, InputBufferSize, InputBufferStride);
 
-				LinearLayer(
-					MemoryCellInstance->RememberGateBuffer.GetData(),
-					Input.GetData(),
-					WeightsWr.GetData(),
-					BiasesWr.GetData(),
-					BatchSize,
-					InputSize + MemorySize,
-					MemorySize,
-					MemorySize,
-					InputStride);
+				const uint32 SubLayerNum = SubLayers.Num();
 
-				// Passthrough Gate
+				// Gather the batch indices used by each sub-layer
 
-				LinearLayer(
-					MemoryCellInstance->PassthroughGateBuffer.GetData(),
-					Input.GetData(),
-					WeightsWp.GetData(),
-					BiasesWp.GetData(),
+				OperatorGatherSubLayerBatchIndicesExclusive(
+					AggregateOrExclusiveInstance->SubLayerBatchIndices,
+					InputBuffer + MaxSubLayerInputSize,
 					BatchSize,
-					InputSize + MemorySize,
-					OutputSize,
-					OutputSize,
-					InputStride);
+					SubLayerNum,
+					InputBufferStride);
 
-				// Memory Update
+				// Evaluate Sublayers
 
-				LinearLayer(
-					MemoryCellInstance->MemoryUpdateBuffer.GetData(),
-					Input.GetData(),
-					WeightsWn.GetData(),
-					BiasesWn.GetData(),
-					BatchSize,
-					InputSize + MemorySize,
-					MemorySize,
-					MemorySize,
-					InputStride);
+				for (uint32 SubLayerIdx = 0; SubLayerIdx < SubLayerNum; SubLayerIdx++)
+				{
+					const uint32 SubLayerBatchSize = AggregateOrExclusiveInstance->SubLayerBatchIndices[SubLayerIdx].Num();
+
+					if (SubLayerBatchSize == 0) { continue; }
+
+					OperatorGather(
+						AggregateOrExclusiveInstance->SubLayerInputBuffers[SubLayerIdx].GetData(),
+						InputBuffer,
+						AggregateOrExclusiveInstance->SubLayerBatchIndices[SubLayerIdx].GetData(),
+						SubLayerBatchSize,
+						SubLayerInputSizes[SubLayerIdx],
+						SubLayerInputSizes[SubLayerIdx],
+						InputBufferStride);
+
+					SubLayers[SubLayerIdx]->Evaluate(
+						AggregateOrExclusiveInstance->SubLayerInstances[SubLayerIdx].Get(),
+						AggregateOrExclusiveInstance->SubLayerOutputBuffers[SubLayerIdx].GetData(),
+						AggregateOrExclusiveInstance->SubLayerInputBuffers[SubLayerIdx].GetData(),
+						SubLayerBatchSize,
+						SubLayerOutputSizes[SubLayerIdx],
+						SubLayerInputSizes[SubLayerIdx],
+						SubLayerOutputSizes[SubLayerIdx],
+						SubLayerInputSizes[SubLayerIdx]);
+
+					Encoders[SubLayerIdx]->Evaluate(
+						AggregateOrExclusiveInstance->EncoderInstances[SubLayerIdx].Get(),
+						AggregateOrExclusiveInstance->SubLayerEncodingBuffers[SubLayerIdx].GetData(),
+						AggregateOrExclusiveInstance->SubLayerOutputBuffers[SubLayerIdx].GetData(),
+						SubLayerBatchSize,
+						OutputEncodingSize,
+						SubLayerOutputSizes[SubLayerIdx],
+						OutputEncodingSize,
+						SubLayerOutputSizes[SubLayerIdx]);
+
+					OperatorScatter(
+						OutputBuffer,
+						AggregateOrExclusiveInstance->SubLayerEncodingBuffers[SubLayerIdx].GetData(),
+						AggregateOrExclusiveInstance->SubLayerBatchIndices[SubLayerIdx].GetData(),
+						SubLayerBatchSize,
+						OutputEncodingSize,
+						OutputBufferStride,
+						OutputEncodingSize);
+				}
 
-				// Update Memory State
+				// Append SubLayer Mask
 
-				MemoryCellUpdateMemory(
-					Output.GetData() + OutputSize,
-					MemoryCellInstance->RememberGateBuffer.GetData(),
-					Input.GetData() + InputSize,
-					MemoryCellInstance->MemoryUpdateBuffer.GetData(),
+				OperatorCopy(
+					OutputBuffer + OutputEncodingSize,
+					InputBuffer + MaxSubLayerInputSize,
 					BatchSize,
-					MemorySize,
-					OutputStride,
-					MemorySize,
-					InputStride,
-					MemorySize);
+					SubLayerNum,
+					OutputBufferStride,
+					InputBufferStride);
 
-				// Output Input Update
+				OperatorNanCheck(OutputBuffer, BatchSize, OutputBufferSize, OutputBufferStride);
+			}
 
-				LinearLayer(
-					MemoryCellInstance->OutputInputUpdateBuffer.GetData(),
-					Input.GetData(),
-					WeightsWz.GetData(),
-					BiasesWz.GetData(),
-					BatchSize,
-					InputSize + MemorySize,
-					OutputSize,
-					OutputSize,
-					InputStride);
+			uint32 OutputEncodingSize = 0;
+			TConstArrayView<uint32> SubLayerInputSizes;
+			TConstArrayView<uint32> SubLayerOutputSizes;
+			TArray<TSharedPtr<ILayer>, TInlineAllocator<32>> SubLayers;
+			TArray<TSharedPtr<ILayer>, TInlineAllocator<32>> Encoders;
 
-				// Output Memory Update
+			uint32 MaxSubLayerInputSize = 0;
+		};
 
-				LinearLayer(
-					MemoryCellInstance->OutputMemoryUpdateBuffer.GetData(),
-					Output.GetData() + OutputSize,
-					WeightsWy.GetData(),
-					BiasesWy.GetData(),
-					BatchSize,
-					MemorySize,
-					OutputSize,
-					OutputSize,
-					OutputStride);
+		FAggregateOrExclusiveLayerInstance::FAggregateOrExclusiveLayerInstance(const FAggregateOrExclusiveLayer& InAggregateOrExclusiveLayer)
+			: AggregateOrExclusiveLayer(InAggregateOrExclusiveLayer)
+		{
+			const uint32 SubLayerNum = AggregateOrExclusiveLayer.SubLayers.Num();
 
-				// Update Final Output
+			SubLayerInstances.Init(nullptr, SubLayerNum);
+			EncoderInstances.Init(nullptr, SubLayerNum);
 
-				MemoryCellUpdateOutput(
-					Output.GetData(),
-					MemoryCellInstance->PassthroughGateBuffer.GetData(),
-					MemoryCellInstance->OutputMemoryUpdateBuffer.GetData(),
-					MemoryCellInstance->OutputInputUpdateBuffer.GetData(),
-					BatchSize,
-					OutputSize,
-					OutputStride,
-					OutputSize,
-					OutputSize,
-					OutputSize);
-			}
+			SubLayerBatchIndices.SetNum(SubLayerNum);
+			SubLayerInputBuffers.SetNum(SubLayerNum);
+			SubLayerOutputBuffers.SetNum(SubLayerNum);
+			SubLayerEncodingBuffers.SetNum(SubLayerNum);
 
-			uint32 InputSize = 0;
-			uint32 OutputSize = 0;
-			uint32 MemorySize = 0;
-			TConstArrayView<float> BiasesWr;
-			TConstArrayView<float> WeightsWr;
-			TConstArrayView<float> BiasesWp;
-			TConstArrayView<float> WeightsWp;
-			TConstArrayView<float> BiasesWn;
-			TConstArrayView<float> WeightsWn;
-			TConstArrayView<float> BiasesWz;
-			TConstArrayView<float> WeightsWz;
-			TConstArrayView<float> BiasesWy;
-			TConstArrayView<float> WeightsWy;
-		};
+			for (uint32 SubLayerIdx = 0; SubLayerIdx < SubLayerNum; SubLayerIdx++)
+			{
+				SubLayerInstances[SubLayerIdx] = AggregateOrExclusiveLayer.SubLayers[SubLayerIdx]->MakeInstance();
+				EncoderInstances[SubLayerIdx] = AggregateOrExclusiveLayer.Encoders[SubLayerIdx]->MakeInstance();
+			}
+		}
 
-		//--------------------------------------------------------------------------
+		void FAggregateOrExclusiveLayerInstance::SetMaxBatchSize(const uint32 MaxBatchSize)
+		{
+			NNE_RUNTIME_BASIC_TRACE_SCOPE(NNE::RuntimeBasic::Private::FAggregateOrExclusiveLayerInstance::SetMaxBatchSize);
 
-		FMemoryCellInstance::FMemoryCellInstance(const FMemoryCellLayer& InMemoryCellLayer) 
-			: MemoryCellLayer(InMemoryCellLayer) {}
+			const uint32 SubLayerNum = AggregateOrExclusiveLayer.SubLayers.Num();
 
-		void FMemoryCellInstance::SetBatchSize(const uint32 BatchSize)
-		{
-			TRACE_CPUPROFILER_EVENT_SCOPE(NNE::RuntimeBasic::Private::FMemoryCellInstance::SetBatchSize);
+			for (uint32 SubLayerIdx = 0; SubLayerIdx < SubLayerNum; SubLayerIdx++)
+			{
+				if (SubLayerInstances[SubLayerIdx]) { SubLayerInstances[SubLayerIdx]->SetMaxBatchSize(MaxBatchSize); }
+				if (EncoderInstances[SubLayerIdx]) { EncoderInstances[SubLayerIdx]->SetMaxBatchSize(MaxBatchSize); }
 
-			RememberGateBuffer.SetNumUninitialized(BatchSize * MemoryCellLayer.MemorySize, false);
-			MemoryUpdateBuffer.SetNumUninitialized(BatchSize * MemoryCellLayer.MemorySize, false);
-			PassthroughGateBuffer.SetNumUninitialized(BatchSize * MemoryCellLayer.OutputSize, false);
-			OutputMemoryUpdateBuffer.SetNumUninitialized(BatchSize * MemoryCellLayer.OutputSize, false);
-			OutputInputUpdateBuffer.SetNumUninitialized(BatchSize * MemoryCellLayer.OutputSize, false);
+				SubLayerBatchIndices[SubLayerIdx].Empty(MaxBatchSize);
+				SubLayerInputBuffers[SubLayerIdx].SetNumUninitialized(MaxBatchSize * AggregateOrExclusiveLayer.SubLayerInputSizes[SubLayerIdx], false);
+				SubLayerOutputBuffers[SubLayerIdx].SetNumUninitialized(MaxBatchSize * AggregateOrExclusiveLayer.SubLayerOutputSizes[SubLayerIdx], false);
+				SubLayerEncodingBuffers[SubLayerIdx].SetNumUninitialized(MaxBatchSize * AggregateOrExclusiveLayer.OutputEncodingSize, false);
+			}
 		}
 
 		//--------------------------------------------------------------------------
 
-		struct FMemoryBackboneLayer;
+		struct FAggregateOrInclusiveLayer;
 
-		struct FMemoryBackboneInstance : public ILayerInstance
+		struct FAggregateOrInclusiveLayerInstance : public ILayerInstance
 		{
-			FMemoryBackboneInstance(const FMemoryBackboneLayer& InMemoryBackboneLayer);
-
-			virtual void SetBatchSize(const uint32 BatchSize) override final;
-
-			const FMemoryBackboneLayer& MemoryBackboneLayer;
-			TArray<float, TInlineAllocator<512>> CellInputBuffer;
-			TArray<float, TInlineAllocator<512>> CellOutputBuffer;
-			TSharedPtr<ILayerInstance> PrefixInstance;
-			TSharedPtr<ILayerInstance> CellInstance;
-			TSharedPtr<ILayerInstance> PostfixInstance;
+			FAggregateOrInclusiveLayerInstance(const FAggregateOrInclusiveLayer& InAggregateOrInclusiveLayer);
+
+			void SetMaxBatchSize(const uint32 MaxBatchSize) override final;
+
+			const FAggregateOrInclusiveLayer& AggregateOrInclusiveLayer;
+			TArray<TSharedPtr<ILayerInstance>, TInlineAllocator<32>> SubLayerInstances;
+			TArray<TSharedPtr<ILayerInstance>, TInlineAllocator<32>> QueryInstances;
+			TArray<TSharedPtr<ILayerInstance>, TInlineAllocator<32>> KeyInstances;
+			TArray<TSharedPtr<ILayerInstance>, TInlineAllocator<32>> ValueInstances;
+
+			TArray<TArray<uint32>, TInlineAllocator<32>> SubLayerBatchIndices;
+			TArray<TArray<float>, TInlineAllocator<32>> SubLayerInputBuffers;
+			TArray<TArray<float>, TInlineAllocator<32>> SubLayerOutputBuffers;
+			TArray<TArray<float>, TInlineAllocator<32>> SubLayerQueryBuffers;
+			TArray<TArray<float>, TInlineAllocator<32>> SubLayerKeyBuffers;
+			TArray<TArray<float>, TInlineAllocator<32>> SubLayerValueBuffers;
+
+			uint32 TotalElementNum = 0;
+			TArray<uint32, TInlineAllocator<32>> ElementAccum;
+			TArray<uint32, TInlineAllocator<32>> ElementNums;
+			TArray<uint32, TInlineAllocator<32>> ElementOffsets;
+
+			TArray<float> AttentionMaxsBuffer;
+			TArray<float> AttentionDenomsBuffer;
+			TArray<float> AttentionBuffer;
+			TArray<float> QueryBuffer;
+			TArray<float> KeyBuffer;
+			TArray<float> ValueBuffer;
 		};
 
-		//--------------------------------------------------------------------------
-
-		struct FMemoryBackboneLayer : public ILayer
+		struct FAggregateOrInclusiveLayer : public ILayer
 		{
-			virtual TSharedPtr<ILayerInstance> MakeInstance() const { return MakeShared<FMemoryBackboneInstance>(*this); };
-			virtual ELayerType GetLayerType() const override final { return ELayerType::MemoryBackbone; }
-			virtual uint32 GetInputSize() const override final { return PrefixInputSize + MemorySize; }
-			virtual uint32 GetOutputSize() const override final { return PostfixOutputSize + MemorySize; }
+			virtual TSharedPtr<ILayerInstance> MakeInstance() const { return MakeShared<FAggregateOrInclusiveLayerInstance>(*this); };
+			virtual ELayerType GetLayerType() const override final { return ELayerType::AggregateOrInclusive; }
+			virtual uint32 GetInputSize() const override final { return TotalSubLayerInputSize + SubLayers.Num(); }
+			virtual uint32 GetOutputSize() const override final { return AttentionHeadNum * OutputEncodingSize + SubLayers.Num(); }
 
 			virtual void SerializationSize(uint64& InOutOffset) const override final
 			{
-				Serialization::Size(InOutOffset, MemorySize);
-				Serialization::Size(InOutOffset, Cell);
-				Serialization::Size(InOutOffset, Prefix);
-				Serialization::Size(InOutOffset, Postfix);
+				Serialization::Size(InOutOffset, (uint32)SubLayers.Num());
+				Serialization::Size(InOutOffset, OutputEncodingSize);
+				Serialization::Size(InOutOffset, AttentionEncodingSize);
+				Serialization::Size(InOutOffset, AttentionHeadNum);
+				Serialization::Size(InOutOffset, SubLayerInputSizes);
+				Serialization::Size(InOutOffset, SubLayerOutputSizes);
+				Serialization::Size(InOutOffset, SubLayers);
+				Serialization::Size(InOutOffset, QueryLayers);
+				Serialization::Size(InOutOffset, KeyLayers);
+				Serialization::Size(InOutOffset, ValueLayers);
 			}
 
 			virtual void SerializationLoad(uint64& InOutOffset, TConstArrayView<uint8> Data) override final
 			{
-				Serialization::Load(InOutOffset, MemorySize, Data);
-				Serialization::Load(InOutOffset, Cell, Data);
-				Serialization::Load(InOutOffset, Prefix, Data);
-				Serialization::Load(InOutOffset, Postfix, Data);
-
-				PrefixInputSize = Prefix->GetInputSize();
-				PrefixOutputSize = Prefix->GetOutputSize();
-				PostfixInputSize = Postfix->GetInputSize();
-				PostfixOutputSize = Postfix->GetOutputSize();
+				uint32 SubLayerNum = 0;
+				Serialization::Load(InOutOffset, SubLayerNum, Data);
+				Serialization::Load(InOutOffset, OutputEncodingSize, Data);
+				Serialization::Load(InOutOffset, AttentionEncodingSize, Data);
+				Serialization::Load(InOutOffset, AttentionHeadNum, Data);
+				Serialization::Load(InOutOffset, SubLayerInputSizes, Data, SubLayerNum);
+				Serialization::Load(InOutOffset, SubLayerOutputSizes, Data, SubLayerNum);
+				SubLayers.Init(nullptr, SubLayerNum);
+				Serialization::Load(InOutOffset, SubLayers, Data);
+				QueryLayers.Init(nullptr, SubLayerNum);
+				Serialization::Load(InOutOffset, QueryLayers, Data);
+				KeyLayers.Init(nullptr, SubLayerNum);
+				Serialization::Load(InOutOffset, KeyLayers, Data);
+				ValueLayers.Init(nullptr, SubLayerNum);
+				Serialization::Load(InOutOffset, ValueLayers, Data);
+				PostLoad();
+			}
+
+			void PostLoad()
+			{
+				const uint32 SubLayerNum = SubLayers.Num();
+				TotalSubLayerInputSize = 0;
+				SubLayerInputOffsets.SetNumUninitialized(SubLayerNum);
+				for (uint32 SubLayerIdx = 0; SubLayerIdx < SubLayerNum; SubLayerIdx++)
+				{
+					SubLayerInputOffsets[SubLayerIdx] = TotalSubLayerInputSize;
+					TotalSubLayerInputSize += SubLayerInputSizes[SubLayerIdx];
+				}
 			}
 
 			virtual void SerializationSave(uint64& InOutOffset, TArrayView<uint8> Data) const override final
 			{
-				Serialization::Save(InOutOffset, MemorySize, Data);
-				Serialization::Save(InOutOffset, Cell, Data);
-				Serialization::Save(InOutOffset, Prefix, Data);
-				Serialization::Save(InOutOffset, Postfix, Data);
+				Serialization::Save(InOutOffset, (uint32)SubLayers.Num(), Data);
+				Serialization::Save(InOutOffset, OutputEncodingSize, Data);
+				Serialization::Save(InOutOffset, AttentionEncodingSize, Data);
+				Serialization::Save(InOutOffset, AttentionHeadNum, Data);
+				Serialization::Save(InOutOffset, SubLayerInputSizes, Data);
+				Serialization::Save(InOutOffset, SubLayerOutputSizes, Data);
+				Serialization::Save(InOutOffset, SubLayers, Data);
+				Serialization::Save(InOutOffset, QueryLayers, Data);
+				Serialization::Save(InOutOffset, KeyLayers, Data);
+				Serialization::Save(InOutOffset, ValueLayers, Data);
 			}
 
 			virtual void Evaluate(
-				TArrayView<float> Output,
-				TConstArrayView<float> Input,
 				ILayerInstance* Instance,
-				uint32 BatchSize,
-				uint32 OutputStride,
-				uint32 InputStride) override final
+				float* OutputBuffer,
+				const float* InputBuffer,
+				const uint32 BatchSize,
+				const uint32 OutputBufferSize,
+				const uint32 InputBufferSize,
+				const uint32 OutputBufferStride,
+				const uint32 InputBufferStride) override final
 			{
-				TRACE_CPUPROFILER_EVENT_SCOPE(NNE::RuntimeBasic::Private::FMemoryBackboneLayer::Evaluate);
+				NNE_RUNTIME_BASIC_TRACE_SCOPE(NNE::RuntimeBasic::Private::FAggregateOrInclusiveLayer::Evaluate);
+				check(OutputBufferSize == GetOutputSize() && InputBufferSize == GetInputSize());
 
-				FMemoryBackboneInstance* MemoryBackboneInstance = StaticCast<FMemoryBackboneInstance*>(Instance);
-				check(MemoryBackboneInstance);
+				FAggregateOrInclusiveLayerInstance* AggregateOrInclusiveInstance = StaticCast<FAggregateOrInclusiveLayerInstance*>(Instance);
+				check(AggregateOrInclusiveInstance);
 
-				// Evaluate Prefix
+				OperatorNanCheck(InputBuffer, BatchSize, InputBufferSize, InputBufferStride);
 
-				Prefix->Evaluate(
-					MemoryBackboneInstance->CellInputBuffer,
-					Input,
-					MemoryBackboneInstance->PrefixInstance.Get(),
-					BatchSize,
-					PrefixOutputSize + MemorySize,
-					InputStride);
+				const uint32 SubLayerNum = SubLayers.Num();
 
-				// Copy Memory State to Cell Input
+				// Count the number of sub-layer used by each item in the batch
 
-				IdentityLayer(
-					MemoryBackboneInstance->CellInputBuffer.GetData() + PrefixOutputSize,
-					Input.GetData() + PrefixInputSize,
+				OperatorAggregateCountElementNum(
+					AggregateOrInclusiveInstance->TotalElementNum,
+					AggregateOrInclusiveInstance->ElementNums.GetData(),
+					AggregateOrInclusiveInstance->ElementOffsets.GetData(),
+					InputBuffer + TotalSubLayerInputSize,
 					BatchSize,
-					MemorySize,
-					PrefixOutputSize + MemorySize,
-					InputStride);
+					SubLayerNum,
+					InputBufferStride);
 
-				// Evaluate Cell
+				// Gather the batch indices used by each sub-layer
 
-				Cell->Evaluate(
-					MemoryBackboneInstance->CellOutputBuffer,
-					MemoryBackboneInstance->CellInputBuffer,
-					MemoryBackboneInstance->CellInstance.Get(),
+				OperatorGatherSubLayerBatchIndicesInclusive(
+					AggregateOrInclusiveInstance->SubLayerBatchIndices,
+					InputBuffer + TotalSubLayerInputSize,
 					BatchSize,
-					PostfixInputSize + MemorySize,
-					PrefixOutputSize + MemorySize);
+					SubLayerNum,
+					InputBufferStride);
+
+				// Evaluate Each Sublayer on the associated batch items
 
-				// Evaluate Postfix
+				for (uint32 SubLayerIdx = 0; SubLayerIdx < SubLayerNum; SubLayerIdx++)
+				{
+					const uint32 SubLayerBatchSize = AggregateOrInclusiveInstance->SubLayerBatchIndices[SubLayerIdx].Num();
+
+					if (SubLayerBatchSize == 0) { continue; }
+
+					OperatorGather(
+						AggregateOrInclusiveInstance->SubLayerInputBuffers[SubLayerIdx].GetData(),
+						InputBuffer + SubLayerInputOffsets[SubLayerIdx],
+						AggregateOrInclusiveInstance->SubLayerBatchIndices[SubLayerIdx].GetData(),
+						SubLayerBatchSize,
+						SubLayerInputSizes[SubLayerIdx],
+						SubLayerInputSizes[SubLayerIdx],
+						InputBufferStride);
+
+					SubLayers[SubLayerIdx]->Evaluate(
+						AggregateOrInclusiveInstance->SubLayerInstances[SubLayerIdx].Get(),
+						AggregateOrInclusiveInstance->SubLayerOutputBuffers[SubLayerIdx].GetData(),
+						AggregateOrInclusiveInstance->SubLayerInputBuffers[SubLayerIdx].GetData(),
+						SubLayerBatchSize,
+						SubLayerOutputSizes[SubLayerIdx],
+						SubLayerInputSizes[SubLayerIdx],
+						SubLayerOutputSizes[SubLayerIdx],
+						SubLayerInputSizes[SubLayerIdx]);
+
+					QueryLayers[SubLayerIdx]->Evaluate(
+						AggregateOrInclusiveInstance->QueryInstances[SubLayerIdx].Get(),
+						AggregateOrInclusiveInstance->SubLayerQueryBuffers[SubLayerIdx].GetData(),
+						AggregateOrInclusiveInstance->SubLayerOutputBuffers[SubLayerIdx].GetData(),
+						SubLayerBatchSize,
+						AttentionHeadNum * AttentionEncodingSize,
+						SubLayerOutputSizes[SubLayerIdx],
+						AttentionHeadNum * AttentionEncodingSize,
+						SubLayerOutputSizes[SubLayerIdx]);
+
+					KeyLayers[SubLayerIdx]->Evaluate(
+						AggregateOrInclusiveInstance->KeyInstances[SubLayerIdx].Get(),
+						AggregateOrInclusiveInstance->SubLayerKeyBuffers[SubLayerIdx].GetData(),
+						AggregateOrInclusiveInstance->SubLayerOutputBuffers[SubLayerIdx].GetData(),
+						SubLayerBatchSize,
+						AttentionHeadNum * AttentionEncodingSize,
+						SubLayerOutputSizes[SubLayerIdx],
+						AttentionHeadNum * AttentionEncodingSize,
+						SubLayerOutputSizes[SubLayerIdx]);
+
+					ValueLayers[SubLayerIdx]->Evaluate(
+						AggregateOrInclusiveInstance->ValueInstances[SubLayerIdx].Get(),
+						AggregateOrInclusiveInstance->SubLayerValueBuffers[SubLayerIdx].GetData(),
+						AggregateOrInclusiveInstance->SubLayerOutputBuffers[SubLayerIdx].GetData(),
+						SubLayerBatchSize,
+						AttentionHeadNum * OutputEncodingSize,
+						SubLayerOutputSizes[SubLayerIdx],
+						AttentionHeadNum * OutputEncodingSize,
+						SubLayerOutputSizes[SubLayerIdx]);
+				}
 
-				Postfix->Evaluate(
-					Output,
-					MemoryBackboneInstance->CellOutputBuffer,
-					MemoryBackboneInstance->PostfixInstance.Get(),
+				// Gather queries, keys, and values from sub-layers into tightly packed element lists
+				// which we can attend over using the ElementNums and ElementOffsets arrays
+
+				OperatorAggregateGatherFromSubLayers(
+					AggregateOrInclusiveInstance->QueryBuffer.GetData(),
+					AggregateOrInclusiveInstance->KeyBuffer.GetData(),
+					AggregateOrInclusiveInstance->ValueBuffer.GetData(),
+					AggregateOrInclusiveInstance->ElementAccum.GetData(),
+					AggregateOrInclusiveInstance->ElementNums.GetData(),
+					AggregateOrInclusiveInstance->ElementOffsets.GetData(),
+					AggregateOrInclusiveInstance->SubLayerBatchIndices,
+					AggregateOrInclusiveInstance->SubLayerQueryBuffers,
+					AggregateOrInclusiveInstance->SubLayerKeyBuffers,
+					AggregateOrInclusiveInstance->SubLayerValueBuffers,
+					BatchSize,
+					AttentionHeadNum * AttentionEncodingSize,
+					AttentionHeadNum * AttentionEncodingSize,
+					AttentionHeadNum * OutputEncodingSize);
+
+				// Compute Attention
+
+				OperatorAggregateDotProductAttention(
+					AggregateOrInclusiveInstance->AttentionBuffer.GetData(),
+					AggregateOrInclusiveInstance->QueryBuffer.GetData(),
+					AggregateOrInclusiveInstance->KeyBuffer.GetData(),
+					AggregateOrInclusiveInstance->TotalElementNum,
+					AttentionEncodingSize,
+					AttentionHeadNum);
+
+				OperatorAggregateSoftmaxPlusOneInplace(
+					AggregateOrInclusiveInstance->AttentionMaxsBuffer.GetData(),
+					AggregateOrInclusiveInstance->AttentionDenomsBuffer.GetData(),
+					AggregateOrInclusiveInstance->AttentionBuffer.GetData(),
+					AggregateOrInclusiveInstance->ElementNums.GetData(),
+					AggregateOrInclusiveInstance->ElementOffsets.GetData(),
+					BatchSize,
+					AttentionHeadNum);
+
+				OperatorAggregateAttentionSum(
+					OutputBuffer,
+					AggregateOrInclusiveInstance->AttentionBuffer.GetData(),
+					AggregateOrInclusiveInstance->ValueBuffer.GetData(),
+					AggregateOrInclusiveInstance->ElementNums.GetData(),
+					AggregateOrInclusiveInstance->ElementOffsets.GetData(),
 					BatchSize,
-					OutputStride,
-					PostfixInputSize + MemorySize);
+					OutputEncodingSize,
+					AttentionHeadNum,
+					OutputBufferStride);
 
-				// Copy Memory State to Output
+				// Append Element Mask
 
-				IdentityLayer(
-					Output.GetData() + PostfixOutputSize,
-					MemoryBackboneInstance->CellOutputBuffer.GetData() + PostfixInputSize,
+				OperatorCopy(
+					OutputBuffer + AttentionHeadNum * OutputEncodingSize,
+					InputBuffer + TotalSubLayerInputSize,
 					BatchSize,
-					MemorySize,
-					OutputStride,
-					PostfixInputSize + MemorySize);
+					SubLayerNum,
+					OutputBufferStride,
+					InputBufferStride);
+
+				OperatorNanCheck(OutputBuffer, BatchSize, OutputBufferSize, OutputBufferStride);
 			}
 
-			uint32 MemorySize = 0;
-			TSharedPtr<ILayer> Cell;
-			TSharedPtr<ILayer> Prefix;
-			TSharedPtr<ILayer> Postfix;
-			uint32 PrefixInputSize = 0;
-			uint32 PrefixOutputSize = 0;
-			uint32 PostfixInputSize = 0;
-			uint32 PostfixOutputSize = 0;
-		};
+			uint32 OutputEncodingSize = 0;
+			uint32 AttentionEncodingSize = 0;
+			uint32 AttentionHeadNum = 0;
+			TConstArrayView<uint32> SubLayerInputSizes;
+			TConstArrayView<uint32> SubLayerOutputSizes;
+			TArray<TSharedPtr<ILayer>, TInlineAllocator<32>> SubLayers;
+			TArray<TSharedPtr<ILayer>, TInlineAllocator<32>> QueryLayers;
+			TArray<TSharedPtr<ILayer>, TInlineAllocator<32>> KeyLayers;
+			TArray<TSharedPtr<ILayer>, TInlineAllocator<32>> ValueLayers;
 
-		//--------------------------------------------------------------------------
+			uint32 TotalSubLayerInputSize = 0;
+			TArray<uint32, TInlineAllocator<32>> SubLayerInputOffsets;
+		};
 
-		FMemoryBackboneInstance::FMemoryBackboneInstance(const FMemoryBackboneLayer& InMemoryBackboneLayer)
-			: MemoryBackboneLayer(InMemoryBackboneLayer)
+		FAggregateOrInclusiveLayerInstance::FAggregateOrInclusiveLayerInstance(const FAggregateOrInclusiveLayer& InAggregateOrInclusiveLayer)
+			: AggregateOrInclusiveLayer(InAggregateOrInclusiveLayer)
 		{
-			CellInstance = MemoryBackboneLayer.Cell->MakeInstance();
-			PrefixInstance = MemoryBackboneLayer.Prefix->MakeInstance();
-			PostfixInstance = MemoryBackboneLayer.Postfix->MakeInstance();
+			const uint32 SubLayerNum = AggregateOrInclusiveLayer.SubLayers.Num();
+
+			SubLayerInstances.Init(nullptr, SubLayerNum);
+			QueryInstances.Init(nullptr, SubLayerNum);
+			KeyInstances.Init(nullptr, SubLayerNum);
+			ValueInstances.Init(nullptr, SubLayerNum);
+
+			SubLayerBatchIndices.SetNum(SubLayerNum);
+			SubLayerInputBuffers.SetNum(SubLayerNum);
+			SubLayerOutputBuffers.SetNum(SubLayerNum);
+			SubLayerQueryBuffers.SetNum(SubLayerNum);
+			SubLayerKeyBuffers.SetNum(SubLayerNum);
+			SubLayerValueBuffers.SetNum(SubLayerNum);
+
+			for (uint32 SubLayerIdx = 0; SubLayerIdx < SubLayerNum; SubLayerIdx++)
+			{
+				SubLayerInstances[SubLayerIdx] = AggregateOrInclusiveLayer.SubLayers[SubLayerIdx]->MakeInstance();
+				QueryInstances[SubLayerIdx] = AggregateOrInclusiveLayer.QueryLayers[SubLayerIdx]->MakeInstance();
+				KeyInstances[SubLayerIdx] = AggregateOrInclusiveLayer.KeyLayers[SubLayerIdx]->MakeInstance();
+				ValueInstances[SubLayerIdx] = AggregateOrInclusiveLayer.ValueLayers[SubLayerIdx]->MakeInstance();
+			}
 		}
 
-		void FMemoryBackboneInstance::SetBatchSize(const uint32 BatchSize)
+		void FAggregateOrInclusiveLayerInstance::SetMaxBatchSize(const uint32 MaxBatchSize)
 		{
-			TRACE_CPUPROFILER_EVENT_SCOPE(NNE::RuntimeBasic::Private::FMemoryBackboneInstance::SetBatchSize);
+			NNE_RUNTIME_BASIC_TRACE_SCOPE(NNE::RuntimeBasic::Private::FAggregateOrInclusiveLayerInstance::SetMaxBatchSize);
+
+			const uint32 SubLayerNum = AggregateOrInclusiveLayer.SubLayers.Num();
+
+			for (uint32 SubLayerIdx = 0; SubLayerIdx < SubLayerNum; SubLayerIdx++)
+			{
+				if (SubLayerInstances[SubLayerIdx]) { SubLayerInstances[SubLayerIdx]->SetMaxBatchSize(MaxBatchSize); }
+				if (QueryInstances[SubLayerIdx]) { QueryInstances[SubLayerIdx]->SetMaxBatchSize(MaxBatchSize); }
+				if (KeyInstances[SubLayerIdx]) { KeyInstances[SubLayerIdx]->SetMaxBatchSize(MaxBatchSize); }
+				if (ValueInstances[SubLayerIdx]) { ValueInstances[SubLayerIdx]->SetMaxBatchSize(MaxBatchSize); }
+
+				SubLayerBatchIndices[SubLayerIdx].Empty(MaxBatchSize);
+				SubLayerInputBuffers[SubLayerIdx].SetNumUninitialized(MaxBatchSize * AggregateOrInclusiveLayer.SubLayerInputSizes[SubLayerIdx], false);
+				SubLayerOutputBuffers[SubLayerIdx].SetNumUninitialized(MaxBatchSize * AggregateOrInclusiveLayer.SubLayerOutputSizes[SubLayerIdx], false);
+				SubLayerQueryBuffers[SubLayerIdx].SetNumUninitialized(MaxBatchSize * AggregateOrInclusiveLayer.AttentionHeadNum * AggregateOrInclusiveLayer.AttentionEncodingSize, false);
+				SubLayerKeyBuffers[SubLayerIdx].SetNumUninitialized(MaxBatchSize * AggregateOrInclusiveLayer.AttentionHeadNum * AggregateOrInclusiveLayer.AttentionEncodingSize, false);
+				SubLayerValueBuffers[SubLayerIdx].SetNumUninitialized(MaxBatchSize * AggregateOrInclusiveLayer.AttentionHeadNum * AggregateOrInclusiveLayer.OutputEncodingSize, false);
+			}
 
-			if (CellInstance) { CellInstance->SetBatchSize(BatchSize); }
-			if (PrefixInstance) { PrefixInstance->SetBatchSize(BatchSize); }
-			if (PostfixInstance) { PostfixInstance->SetBatchSize(BatchSize); }
+			TotalElementNum = 0;
+			ElementAccum.SetNumUninitialized(MaxBatchSize);
+			ElementNums.SetNumUninitialized(MaxBatchSize);
+			ElementOffsets.SetNumUninitialized(MaxBatchSize);
 
-			CellInputBuffer.SetNumUninitialized(BatchSize * (MemoryBackboneLayer.PrefixOutputSize + MemoryBackboneLayer.MemorySize), false);
-			CellOutputBuffer.SetNumUninitialized(BatchSize * (MemoryBackboneLayer.PostfixInputSize + MemoryBackboneLayer.MemorySize), false);
+			AttentionMaxsBuffer.SetNumUninitialized(MaxBatchSize * AggregateOrInclusiveLayer.AttentionHeadNum, false);
+			AttentionDenomsBuffer.SetNumUninitialized(MaxBatchSize * AggregateOrInclusiveLayer.AttentionHeadNum, false);
+			AttentionBuffer.SetNumUninitialized(MaxBatchSize * SubLayerNum * AggregateOrInclusiveLayer.AttentionHeadNum, false);
+			QueryBuffer.SetNumUninitialized(MaxBatchSize * SubLayerNum * AggregateOrInclusiveLayer.AttentionHeadNum * AggregateOrInclusiveLayer.AttentionEncodingSize);
+			KeyBuffer.SetNumUninitialized(MaxBatchSize * SubLayerNum * AggregateOrInclusiveLayer.AttentionHeadNum * AggregateOrInclusiveLayer.AttentionEncodingSize);
+			ValueBuffer.SetNumUninitialized(MaxBatchSize * SubLayerNum * AggregateOrInclusiveLayer.AttentionHeadNum * AggregateOrInclusiveLayer.OutputEncodingSize);
 		}
 
+		//--------------------------------------------------------------------------
+		// Layer Serialization
 		//--------------------------------------------------------------------------
 
 		namespace Serialization
@@ -1667,6 +3314,11 @@ namespace UE::NNE::RuntimeBasic
 				InLayer->SerializationSize(InOutOffset);
 			}
 
+			static inline void Size(uint64& InOutOffset, const TConstArrayView<TSharedPtr<ILayer>> InLayers)
+			{
+				for (const TSharedPtr<ILayer>& Layer : InLayers) { Serialization::Size(InOutOffset, Layer); }
+			}
+
 			static inline void Load(uint64& InOutOffset, TSharedPtr<ILayer>& OutLayer, TConstArrayView<uint8> Data)
 			{
 				uint32 LayerTypeId = (uint32)ELayerType::Invalid;
@@ -1689,7 +3341,12 @@ namespace UE::NNE::RuntimeBasic
 					case ELayerType::TanH: OutLayer = MakeShared<FTanHLayer>(); break;
 					case ELayerType::PReLU: OutLayer = MakeShared<FPReLULayer>(); break;
 					case ELayerType::MemoryCell: OutLayer = MakeShared<FMemoryCellLayer>(); break;
-					case ELayerType::MemoryBackbone: OutLayer = MakeShared<FMemoryBackboneLayer>(); break;
+					case ELayerType::Copy: OutLayer = MakeShared<FCopyLayer>(); break;
+					case ELayerType::Concat: OutLayer = MakeShared<FConcatLayer>(); break;
+					case ELayerType::Array: OutLayer = MakeShared<FArrayLayer>(); break;
+					case ELayerType::AggregateSet: OutLayer = MakeShared<FAggregateSetLayer>(); break;
+					case ELayerType::AggregateOrExclusive: OutLayer = MakeShared<FAggregateOrExclusiveLayer>(); break;
+					case ELayerType::AggregateOrInclusive: OutLayer = MakeShared<FAggregateOrInclusiveLayer>(); break;
 					default: checkf(false, TEXT("Unknown Layer Id %i"), LayerTypeId);
 					}
 				}
@@ -1697,14 +3354,26 @@ namespace UE::NNE::RuntimeBasic
 				OutLayer->SerializationLoad(InOutOffset, Data);
 			}
 
+			static inline void Load(uint64& InOutOffset, TArrayView<TSharedPtr<ILayer>> OutLayers, TConstArrayView<uint8> Data)
+			{
+				for (TSharedPtr<ILayer>& Layer : OutLayers) { Serialization::Load(InOutOffset, Layer, Data); }
+			}
+
 			static inline void Save(uint64& InOutOffset, const TSharedPtr<ILayer>& InLayer, TArrayView<uint8> Data)
 			{
 				Serialization::Save(InOutOffset, (uint32)InLayer->GetLayerType(), Data);
 				InLayer->SerializationSave(InOutOffset, Data);
 			}
+
+			static inline void Save(uint64& InOutOffset, const TConstArrayView<TSharedPtr<ILayer>> InLayers, TArrayView<uint8> Data)
+			{
+				for (const TSharedPtr<ILayer>& Layer : InLayers) { Serialization::Save(InOutOffset, Layer, Data); }
+			}
 		}
 	}
 
+	//--------------------------------------------------------------------------
+	// NNE Interface Implementation
 	//--------------------------------------------------------------------------
 
 	FModelInstanceCPU::FModelInstanceCPU(const TSharedPtr<FModelCPU>& InModel)
@@ -1716,7 +3385,7 @@ namespace UE::NNE::RuntimeBasic
 
 	int FModelInstanceCPU::SetInputTensorShapes(TConstArrayView<FTensorShape> InInputShapes)
 	{
-		TRACE_CPUPROFILER_EVENT_SCOPE(NNE::RuntimeBasic::FModelInstanceCPU::SetInputTensorShapes);
+		NNE_RUNTIME_BASIC_TRACE_SCOPE(NNE::RuntimeBasic::FModelInstanceCPU::SetInputTensorShapes);
 
 		if (!ensureMsgf(InInputShapes.Num() == 1, TEXT("Basic CPU Inference only supports single input tensor.")))
 		{
@@ -1749,7 +3418,7 @@ namespace UE::NNE::RuntimeBasic
 
 		if (Instance)
 		{
-			Instance->SetBatchSize(BatchSize);
+			Instance->SetMaxBatchSize(BatchSize);
 		}
 
 		return 0;
@@ -1757,7 +3426,7 @@ namespace UE::NNE::RuntimeBasic
 
 	int FModelInstanceCPU::RunSync(TConstArrayView<FTensorBindingCPU> InInputBindings, TConstArrayView<FTensorBindingCPU> InOutputBindings)
 	{
-		TRACE_CPUPROFILER_EVENT_SCOPE(NNE::RuntimeBasic::FModelInstanceCPU::RunSync);
+		NNE_RUNTIME_BASIC_TRACE_SCOPE(NNE::RuntimeBasic::FModelInstanceCPU::RunSync);
 
 		if (!ensureMsgf(BatchSize > 0, TEXT("SetInputTensorShapes must be run before RunSync")))
 		{
@@ -1785,19 +3454,18 @@ namespace UE::NNE::RuntimeBasic
 		}
 
 		Model->Layer->Evaluate(
-			TArrayView<float>((float*)InOutputBindings[0].Data, BatchSize * OutputSize),
-			TConstArrayView<float>((const float*)InInputBindings[0].Data, BatchSize * InputSize),
 			Instance.Get(),
+			(float*)InOutputBindings[0].Data,
+			(const float*)InInputBindings[0].Data,
 			BatchSize,
 			OutputSize,
+			InputSize,
+			OutputSize,
 			InputSize);
 
 		return 0;
 	}
 
-
-	//--------------------------------------------------------------------------
-
 	uint32 FModelCPU::ModelMagicNumber = 0x0BA51C01;
 	uint32 FModelCPU::ModelVersionNumber = 1;
 
@@ -1818,7 +3486,7 @@ namespace UE::NNE::RuntimeBasic
 
 	bool FModelCPU::SerializationLoad(uint64& InOutOffset, TConstArrayView<uint8> Data)
 	{
-		TRACE_CPUPROFILER_EVENT_SCOPE(NNE::RuntimeBasic::FModelCPU::SerializationLoad);
+		NNE_RUNTIME_BASIC_TRACE_SCOPE(NNE::RuntimeBasic::FModelCPU::SerializationLoad);
 
 		checkf(InOutOffset % 64 == 0, 
 			TEXT("Model must be aligned to 64 bytes because there must be no padding before magic number."));
@@ -1846,7 +3514,7 @@ namespace UE::NNE::RuntimeBasic
 
 	void FModelCPU::SerializationSave(uint64& InOutOffset, TArrayView<uint8> Data) const
 	{
-		TRACE_CPUPROFILER_EVENT_SCOPE(NNE::RuntimeBasic::FModelCPU::SerializationSave);
+		NNE_RUNTIME_BASIC_TRACE_SCOPE(NNE::RuntimeBasic::FModelCPU::SerializationSave);
 
 		checkf(InOutOffset % 64 == 0,
 			TEXT("Model must be aligned to 64 bytes because there must be no padding before magic number."));
@@ -1856,25 +3524,44 @@ namespace UE::NNE::RuntimeBasic
 		Private::Serialization::Save(InOutOffset, Layer, Data);
 	}
 
+	//--------------------------------------------------------------------------
+	// Builder
 	//--------------------------------------------------------------------------
 
-	FSequentialModelBuilder::FSequentialModelBuilder()
+	namespace Private
 	{
-		Model = MakeShared<FModelCPU>();
-		Model->Layer = MakeShared<Private::FSequenceLayer>();
+		static inline float UniformToGaussian(const float R1, const float R2)
+		{
+			return FMath::Sqrt(-2.0f * FMath::Loge(FMath::Max(R1, UE_SMALL_NUMBER))) * FMath::Cos(R2 * UE_TWO_PI);
+		}
 	}
 
-	void FSequentialModelBuilder::AddLinear(uint32 InputSize, uint32 OutputSize, TConstArrayView<float> Weights, TConstArrayView<float> Biases)
+	FModelBuilder::FModelBuilder(int32 Seed) : Rng(Seed) {};
+
+	FModelBuilderElement::FModelBuilderElement() = default;
+	FModelBuilderElement::FModelBuilderElement(const TSharedPtr<Private::ILayer>& Ptr) : Layer(Ptr) {}
+	FModelBuilderElement::~FModelBuilderElement() = default;
+
+	int32 FModelBuilderElement::GetInputSize() const
 	{
-		Private::FSequenceLayer* Sequence = StaticCast<Private::FSequenceLayer*>(Model->Layer.Get());
+		check(Layer);
+		return Layer->GetInputSize();
+	}
 
-		if (Sequence->Layers.Num() > 0)
-		{
-			check(Sequence->Layers.Last()->GetOutputSize() == InputSize);
-		}
+	int32 FModelBuilderElement::GetOutputSize() const
+	{
+		check(Layer);
+		return Layer->GetOutputSize();
+	}
 
-		check(Weights.Num() == InputSize * OutputSize);
+	FModelBuilderElement FModelBuilder::MakeLinear(
+		const uint32 InputSize,
+		const uint32 OutputSize,
+		const TConstArrayView<float> Weights,
+		const TConstArrayView<float> Biases)
+	{
 		check(Biases.Num() == OutputSize);
+		check(Weights.Num() == InputSize * OutputSize);
 
 		const TSharedPtr<Private::FLinearLayer> LinearLayer = MakeShared<Private::FLinearLayer>();
 		LinearLayer->InputSize = InputSize;
@@ -1882,20 +3569,30 @@ namespace UE::NNE::RuntimeBasic
 		LinearLayer->Biases = Biases;
 		LinearLayer->Weights = Weights;
 
-		Sequence->Layers.Emplace(LinearLayer);
+		return StaticCastSharedPtr<Private::ILayer>(LinearLayer);
 	}
 
-	void FSequentialModelBuilder::AddMultiLinear(uint32 InputSize, uint32 OutputSize, uint32 BlockNum, TConstArrayView<float> Weights, TConstArrayView<float> Biases)
+	FModelBuilderElement FModelBuilder::MakeLinearWithRandomKaimingWeights(
+		const uint32 InputSize,
+		const uint32 OutputSize,
+		const float WeightScale)
 	{
-		Private::FSequenceLayer* Sequence = StaticCast<Private::FSequenceLayer*>(Model->Layer.Get());
-
-		if (Sequence->Layers.Num() > 0)
-		{
-			check(Sequence->Layers.Last()->GetOutputSize() == InputSize * BlockNum);
-		}
+		return MakeLinear(
+			InputSize,
+			OutputSize,
+			MakeWeightsRandomKaiming(InputSize, OutputSize, WeightScale),
+			MakeWeightsZero(OutputSize));
+	}
 
-		check(Weights.Num() == InputSize * OutputSize * BlockNum);
+	FModelBuilderElement FModelBuilder::MakeMultiLinear(
+		const uint32 InputSize, 
+		const uint32 OutputSize, 
+		const uint32 BlockNum, 
+		const TConstArrayView<float> Weights, 
+		const TConstArrayView<float> Biases)
+	{
 		check(Biases.Num() == OutputSize * BlockNum);
+		check(Weights.Num() == InputSize * OutputSize * BlockNum);
 
 		const TSharedPtr<Private::FMultiLinearLayer> MultiLinearLayer = MakeShared<Private::FMultiLinearLayer>();
 		MultiLinearLayer->InputSize = InputSize;
@@ -1904,203 +3601,469 @@ namespace UE::NNE::RuntimeBasic
 		MultiLinearLayer->Biases = Biases;
 		MultiLinearLayer->Weights = Weights;
 
-		Sequence->Layers.Emplace(MultiLinearLayer);
+		return StaticCastSharedPtr<Private::ILayer>(MultiLinearLayer);
 	}
 
-	void FSequentialModelBuilder::AddReLU()
+	FModelBuilderElement FModelBuilder::MakeNormalize(
+		const uint32 InputOutputSize,
+		const TConstArrayView<float> Mean,
+		const TConstArrayView<float> Std)
 	{
-		Private::FSequenceLayer* Sequence = StaticCast<Private::FSequenceLayer*>(Model->Layer.Get());
-
-		checkf(Sequence->Layers.Num() > 0, TEXT("Cannot add activation on initial layer because size is unknown"));
+		const TSharedPtr<Private::FNormalizeLayer> NormalizeLayer = MakeShared<Private::FNormalizeLayer>();
+		NormalizeLayer->InputOutputSize = InputOutputSize;
+		NormalizeLayer->Mean = Mean;
+		NormalizeLayer->Std = Std;
 
-		const TSharedPtr<Private::FReLULayer> ActivationLayer = MakeShared<Private::FReLULayer>();
-		ActivationLayer->InputOutputSize = Sequence->Layers.Last()->GetOutputSize();
-		Sequence->Layers.Emplace(ActivationLayer);
+		return StaticCastSharedPtr<Private::ILayer>(NormalizeLayer);
 	}
 
-	void FSequentialModelBuilder::AddELU()
+	FModelBuilderElement FModelBuilder::MakeDenormalize(
+		const uint32 InputOutputSize,
+		const TConstArrayView<float> Mean,
+		const TConstArrayView<float> Std)
 	{
-		Private::FSequenceLayer* Sequence = StaticCast<Private::FSequenceLayer*>(Model->Layer.Get());
+		const TSharedPtr<Private::FDenormalizeLayer> DenormalizeLayer = MakeShared<Private::FDenormalizeLayer>();
+		DenormalizeLayer->InputOutputSize = InputOutputSize;
+		DenormalizeLayer->Mean = Mean;
+		DenormalizeLayer->Std = Std;
 
-		checkf(Sequence->Layers.Num() > 0, TEXT("Cannot add activation on initial layer because size is unknown"));
+		return StaticCastSharedPtr<Private::ILayer>(DenormalizeLayer);
+	}
 
-		const TSharedPtr<Private::FELULayer> ActivationLayer = MakeShared<Private::FELULayer>();
-		ActivationLayer->InputOutputSize = Sequence->Layers.Last()->GetOutputSize();
-		Sequence->Layers.Emplace(ActivationLayer);
+	FModelBuilderElement FModelBuilder::MakeReLU(const uint32 InputOutputSize)
+	{
+		const TSharedPtr<Private::FReLULayer> Layer = MakeShared<Private::FReLULayer>();
+		Layer->InputOutputSize = InputOutputSize;
+		return StaticCastSharedPtr<Private::ILayer>(Layer);
 	}
 
-	void FSequentialModelBuilder::AddTanH()
+	FModelBuilderElement FModelBuilder::MakeELU(const uint32 InputOutputSize)
 	{
-		Private::FSequenceLayer* Sequence = StaticCast<Private::FSequenceLayer*>(Model->Layer.Get());
+		const TSharedPtr<Private::FELULayer> Layer = MakeShared<Private::FELULayer>();
+		Layer->InputOutputSize = InputOutputSize;
+		return StaticCastSharedPtr<Private::ILayer>(Layer);
+	}
 
-		checkf(Sequence->Layers.Num() > 0, TEXT("Cannot add activation on initial layer because size is unknown"));
+	FModelBuilderElement FModelBuilder::MakeTanH(const uint32 InputOutputSize)
+	{
+		const TSharedPtr<Private::FTanHLayer> Layer = MakeShared<Private::FTanHLayer>();
+		Layer->InputOutputSize = InputOutputSize;
+		return StaticCastSharedPtr<Private::ILayer>(Layer);
+	}
 
-		const TSharedPtr<Private::FTanHLayer> ActivationLayer = MakeShared<Private::FTanHLayer>();
-		ActivationLayer->InputOutputSize = Sequence->Layers.Last()->GetOutputSize();
-		Sequence->Layers.Emplace(ActivationLayer);
+	FModelBuilderElement FModelBuilder::MakeCopy(const uint32 InputOutputSize)
+	{
+		const TSharedPtr<Private::FCopyLayer> Layer = MakeShared<Private::FCopyLayer>();
+		Layer->InputOutputSize = InputOutputSize;
+		return StaticCastSharedPtr<Private::ILayer>(Layer);
 	}
 
-	void FSequentialModelBuilder::AddPReLU(TConstArrayView<float> Alpha)
+	FModelBuilderElement FModelBuilder::MakeActivation(const uint32 InputOutputSize, const EActivationFunction ActivationFunction)
 	{
-		Private::FSequenceLayer* Sequence = StaticCast<Private::FSequenceLayer*>(Model->Layer.Get());
+		switch (ActivationFunction)
+		{
+		case EActivationFunction::ReLU: return MakeReLU(InputOutputSize);
+		case EActivationFunction::ELU: return MakeELU(InputOutputSize);
+		case EActivationFunction::TanH: return MakeTanH(InputOutputSize);
+		default:
+			checkf(false, TEXT("Unknown Activation Function"));
+			return MakeReLU(InputOutputSize);
+		}
+	}
 
-		checkf(Sequence->Layers.Num() > 0, TEXT("Cannot add activation on initial layer because size is unknown"));
-		check(Alpha.Num() == Sequence->Layers.Last()->GetOutputSize());
+	FModelBuilderElement FModelBuilder::MakePReLU(const uint32 InputOutputSize, TConstArrayView<float> Alpha)
+	{
+		check(Alpha.Num() == InputOutputSize);
 
 		const TSharedPtr<Private::FPReLULayer> ActivationLayer = MakeShared<Private::FPReLULayer>();
-		ActivationLayer->InputOutputSize = Sequence->Layers.Last()->GetOutputSize();
+		ActivationLayer->InputOutputSize = InputOutputSize;
 		ActivationLayer->Alpha = Alpha;
 
-		Sequence->Layers.Emplace(ActivationLayer);
+		return StaticCastSharedPtr<Private::ILayer>(ActivationLayer);
 	}
 
-	uint64 FSequentialModelBuilder::GetWriteByteNum() const
+	FModelBuilderElement FModelBuilder::MakeSequence(const TConstArrayView<FModelBuilderElement> Elements)
 	{
-		uint64 Offset = 0;
-		Model->SerializationSize(Offset);
-		return Offset;
+		const TSharedPtr<Private::FSequenceLayer> SequenceLayer = MakeShared<Private::FSequenceLayer>();
+		SequenceLayer->Layers.Reserve(Elements.Num());
+
+		for (const FModelBuilderElement& Element : Elements)
+		{
+			SequenceLayer->Layers.Emplace(Element.Layer);
+		}
+
+		for (int32 LayerIdx = 1; LayerIdx < Elements.Num(); LayerIdx++)
+		{
+			const int32 PrevLayerOutputSize = SequenceLayer->Layers[LayerIdx - 1]->GetOutputSize();
+			const int32 NextLayerInputSize = SequenceLayer->Layers[LayerIdx - 0]->GetInputSize();
+			checkf(PrevLayerOutputSize == NextLayerInputSize, TEXT("Sequence Layer Dimensions don't match. Output %i vs Input %i."), PrevLayerOutputSize, NextLayerInputSize);
+		}
+
+		return StaticCastSharedPtr<Private::ILayer>(SequenceLayer);
 	}
 
-	void FSequentialModelBuilder::WriteAndReset(TArrayView<uint8> OutBytes)
+	FModelBuilderElement FModelBuilder::MakeMLPWithRandomKaimingWeights(
+		const uint32 InputSize,
+		const uint32 OutputSize,
+		const uint32 HiddenSize,
+		const uint32 LayerNum,
+		const EActivationFunction ActivationFunction,
+		const bool bActivationOnFinalLayer)
 	{
-		check((uint64)OutBytes.Num() >= GetWriteByteNum());
+		check(LayerNum >= 2);
 
-		// Zero to ensure any padding due to alignment is always zero
-		FMemory::Memzero(OutBytes.GetData(), OutBytes.Num());
+		TArray<FModelBuilderElement, TInlineAllocator<32>> Layers;
+		Layers.Reserve(2 * LayerNum - (bActivationOnFinalLayer ? 0 : 1));
 
-		uint64 Offset = 0;
-		Model->SerializationSave(Offset, OutBytes);
+		for (uint32 LayerIdx = 0; LayerIdx < LayerNum; LayerIdx++)
+		{
+			const uint32 LayerInputSize = LayerIdx == 0 ? InputSize : HiddenSize;
+			const uint32 LayerOuputSize = LayerIdx == LayerNum - 1 ? OutputSize : HiddenSize;
 
-		Private::FSequenceLayer* Sequence = StaticCast<Private::FSequenceLayer*>(Model->Layer.Get());
-		Sequence->Layers.Empty();
+			Layers.Emplace(MakeLinearWithRandomKaimingWeights(LayerInputSize, LayerOuputSize));
+			
+			if (bActivationOnFinalLayer || LayerIdx != LayerNum - 1)
+			{
+				Layers.Emplace(MakeActivation(LayerOuputSize, ActivationFunction));
+			}
+		}
+		
+		return MakeSequence(Layers);
 	}
 
-	//--------------------------------------------------------------------------
-
-	FMemoryBackboneModelBuilder::FMemoryBackboneModelBuilder()
+	FModelBuilderElement FModelBuilder::MakeMemoryCell(
+		const uint32 InputNum,
+		const uint32 OutputNum,
+		const uint32 MemoryNum,
+		const FModelBuilderElement& RememberLayer,
+		const FModelBuilderElement& PassthroughLayer,
+		const FModelBuilderElement& MemoryUpdateLayer,
+		const FModelBuilderElement& OutputInputUpdateLayer,
+		const FModelBuilderElement& OutputMemoryUpdateLayer)
 	{
-		Model = MakeShared<FModelCPU>();
+		check(RememberLayer.GetInputSize() == InputNum + MemoryNum);
+		check(RememberLayer.GetOutputSize() == MemoryNum);
+		check(PassthroughLayer.GetInputSize() == InputNum + MemoryNum);
+		check(PassthroughLayer.GetOutputSize() == OutputNum);
+		check(MemoryUpdateLayer.GetInputSize() == InputNum + MemoryNum);
+		check(MemoryUpdateLayer.GetOutputSize() == MemoryNum);
+		check(OutputInputUpdateLayer.GetInputSize() == InputNum + MemoryNum);
+		check(OutputInputUpdateLayer.GetOutputSize() == OutputNum);
+		check(OutputMemoryUpdateLayer.GetInputSize() == MemoryNum);
+		check(OutputMemoryUpdateLayer.GetOutputSize() == OutputNum);
+
+		TSharedPtr<Private::FMemoryCellLayer> CellLayer = MakeShared<Private::FMemoryCellLayer>();
+		CellLayer->InputSize = InputNum;
+		CellLayer->OutputSize = OutputNum;
+		CellLayer->MemorySize = MemoryNum;
+		CellLayer->RememberLayer = RememberLayer.Layer;
+		CellLayer->PassthroughLayer = PassthroughLayer.Layer;
+		CellLayer->MemoryUpdateLayer = MemoryUpdateLayer.Layer;
+		CellLayer->OutputInputUpdateLayer = OutputInputUpdateLayer.Layer;
+		CellLayer->OutputMemoryUpdateLayer = OutputMemoryUpdateLayer.Layer;
+		
+		return StaticCastSharedPtr<Private::ILayer>(CellLayer);
 	}
 
-	void FMemoryBackboneModelBuilder::BuildEmptyModel(
+	FModelBuilderElement FModelBuilder::MakeMemoryCellWithLinearRandomKaimingWeights(
 		const uint32 InputNum,
 		const uint32 OutputNum,
 		const uint32 MemoryNum,
-		const uint32 HiddenUnitNum,
-		const uint32 PrefixLayerNum,
-		const uint32 PostfixLayerNum)
+		const float WeightScale)
 	{
-		// Allocate enough zeros for every weight matrix used
+		return MakeMemoryCell(
+			InputNum,
+			OutputNum,
+			MemoryNum,
+			MakeLinearWithRandomKaimingWeights(InputNum + MemoryNum, MemoryNum, WeightScale),
+			MakeLinearWithRandomKaimingWeights(InputNum + MemoryNum, OutputNum, WeightScale),
+			MakeLinearWithRandomKaimingWeights(InputNum + MemoryNum, MemoryNum, WeightScale),
+			MakeLinearWithRandomKaimingWeights(InputNum + MemoryNum, OutputNum, WeightScale),
+			MakeLinearWithRandomKaimingWeights(MemoryNum, OutputNum, WeightScale));
+	}
 
-		const uint32 MaxZeroNumPrefix = FMath::Max(InputNum * HiddenUnitNum, HiddenUnitNum * HiddenUnitNum);
-		const uint32 MaxZeroNumPostfix = FMath::Max(HiddenUnitNum * HiddenUnitNum, OutputNum * HiddenUnitNum);
-		const uint32 MaxZeroNumCell = FMath::Max((HiddenUnitNum + MemoryNum) * MemoryNum, (HiddenUnitNum + MemoryNum) * HiddenUnitNum);
+	FModelBuilderElement FModelBuilder::MakeMemoryBackbone(
+		const uint32 MemoryNum,
+		const FModelBuilderElement& Prefix,
+		const FModelBuilderElement& Cell,
+		const FModelBuilderElement& Postfix)
+	{
+		check(Prefix.GetOutputSize() == Cell.GetInputSize() - MemoryNum);
+		check(Postfix.GetInputSize() == Cell.GetOutputSize() - MemoryNum);
+
+		return MakeSequence({
+			MakeConcat({
+				Prefix,
+				MakeCopy(MemoryNum)
+				}),
+			Cell,
+			MakeConcat({
+				Postfix,
+				MakeCopy(MemoryNum)
+			})
+		});
+	}
 
-		ZerosData.Init(0.0f, FMath::Max(FMath::Max(MaxZeroNumPrefix, MaxZeroNumPostfix), MaxZeroNumCell));
-		TConstArrayView<float> ZerosView = ZerosData;
+	FModelBuilderElement FModelBuilder::MakeConcat(const TConstArrayView<FModelBuilderElement> Elements)
+	{
+		const int32 LayerNum = Elements.Num();
 
-		// Create Cell Layer
+		const TSharedPtr<Private::FConcatLayer> ConcatLayer = MakeShared<Private::FConcatLayer>();
+		ConcatLayer->InputSizes = MakeSizesLayerInputs(Elements);
+		ConcatLayer->OutputSizes = MakeSizesLayerOutputs(Elements);
+		ConcatLayer->Layers.Reserve(LayerNum);
 
-		TSharedRef<Private::FMemoryCellLayer> CellLayer = MakeShared<Private::FMemoryCellLayer>();
-		CellLayer->InputSize = HiddenUnitNum;
-		CellLayer->MemorySize = MemoryNum;
-		CellLayer->OutputSize = HiddenUnitNum;
-		CellLayer->BiasesWr = ZerosView.Slice(0, MemoryNum);
-		CellLayer->WeightsWr = ZerosView.Slice(0, (HiddenUnitNum + MemoryNum) * MemoryNum);
-		CellLayer->BiasesWp = ZerosView.Slice(0, HiddenUnitNum);
-		CellLayer->WeightsWp = ZerosView.Slice(0, (HiddenUnitNum + MemoryNum) * HiddenUnitNum);
-		CellLayer->BiasesWn = ZerosView.Slice(0, MemoryNum);
-		CellLayer->WeightsWn = ZerosView.Slice(0, (HiddenUnitNum + MemoryNum) * MemoryNum);
-		CellLayer->BiasesWz = ZerosView.Slice(0, HiddenUnitNum);
-		CellLayer->WeightsWz = ZerosView.Slice(0, (HiddenUnitNum + MemoryNum) * HiddenUnitNum);
-		CellLayer->BiasesWy = ZerosView.Slice(0, HiddenUnitNum);
-		CellLayer->WeightsWy = ZerosView.Slice(0, MemoryNum * HiddenUnitNum);
+		for (const FModelBuilderElement& Element : Elements)
+		{
+			ConcatLayer->Layers.Emplace(Element.Layer);
+		}
 
-		// Create Prefix Network
+		ConcatLayer->PostLoad();
 
-		check(PrefixLayerNum >= 1);
+		return StaticCastSharedPtr<Private::ILayer>(ConcatLayer);
+	}
 
-		TSharedRef<Private::FSequenceLayer> PrefixLayer = MakeShared<Private::FSequenceLayer>();
-		PrefixLayer->Layers.Reserve(PrefixLayerNum * 2);
+	FModelBuilderElement FModelBuilder::MakeArray(const uint32 ElementNum, const FModelBuilderElement& SubLayer)
+	{
+		const TSharedPtr<Private::FArrayLayer> ArrayLayer = MakeShared<Private::FArrayLayer>();
+		ArrayLayer->ElementNum = ElementNum;
+		ArrayLayer->ElementInputSize = SubLayer.GetInputSize();
+		ArrayLayer->ElementOutputSize = SubLayer.GetOutputSize();
+		ArrayLayer->SubLayer = SubLayer.Layer;
 
-		for (uint32 LayerIdx = 0; LayerIdx < PrefixLayerNum; LayerIdx++)
-		{
-			const TSharedPtr<Private::FLinearLayer> LinearLayer = MakeShared<Private::FLinearLayer>();
-			LinearLayer->InputSize = LayerIdx == 0 ? InputNum : HiddenUnitNum;
-			LinearLayer->OutputSize = HiddenUnitNum;
-			LinearLayer->Biases = ZerosView.Slice(0, LinearLayer->OutputSize);
-			LinearLayer->Weights = ZerosView.Slice(0, LinearLayer->InputSize * LinearLayer->OutputSize);
+		return StaticCastSharedPtr<Private::ILayer>(ArrayLayer);
+	}
 
-			PrefixLayer->Layers.Emplace(LinearLayer);
+	FModelBuilderElement FModelBuilder::MakeAggregateSet(
+		const uint32 MaxElementNum,
+		const uint32 OutputEncodingSize,
+		const uint32 AttentionEncodingSize,
+		const uint32 AttentionHeadNum,
+		const FModelBuilderElement& SubLayer,
+		const FModelBuilderElement& QueryLayer,
+		const FModelBuilderElement& KeyLayer,
+		const FModelBuilderElement& ValueLayer)
+	{
+		check(SubLayer.GetOutputSize() == QueryLayer.GetInputSize());
+		check(SubLayer.GetOutputSize() == KeyLayer.GetInputSize());
+		check(SubLayer.GetOutputSize() == ValueLayer.GetInputSize());
+		check(QueryLayer.GetOutputSize() == AttentionHeadNum * AttentionEncodingSize);
+		check(KeyLayer.GetOutputSize() == AttentionHeadNum * AttentionEncodingSize);
+		check(ValueLayer.GetOutputSize() == AttentionHeadNum * OutputEncodingSize);
+
+		const TSharedPtr<Private::FAggregateSetLayer> SetLayer = MakeShared<Private::FAggregateSetLayer>();
+		SetLayer->MaxElementNum = MaxElementNum;
+		SetLayer->ElementInputSize = SubLayer.GetInputSize();
+		SetLayer->ElementOutputSize = SubLayer.GetOutputSize();
+		SetLayer->OutputEncodingSize = OutputEncodingSize;
+		SetLayer->AttentionEncodingSize = AttentionEncodingSize;
+		SetLayer->AttentionHeadNum = AttentionHeadNum;
+		SetLayer->SubLayer = SubLayer.Layer;
+		SetLayer->QueryLayer = QueryLayer.Layer;
+		SetLayer->KeyLayer = KeyLayer.Layer;
+		SetLayer->ValueLayer = ValueLayer.Layer;
+
+		return StaticCastSharedPtr<Private::ILayer>(SetLayer);
+	}
 
-			const TSharedPtr<Private::FELULayer> ActivationLayer = MakeShared<Private::FELULayer>();
-			ActivationLayer->InputOutputSize = LinearLayer->OutputSize;
+	FModelBuilderElement FModelBuilder::MakeAggregateOrExclusive(
+		const uint32 OutputEncodingSize,
+		const TConstArrayView<FModelBuilderElement> SubLayers,
+		const TConstArrayView<FModelBuilderElement> Encoders)
+	{
+		check(SubLayers.Num() == Encoders.Num());
 
-			PrefixLayer->Layers.Emplace(ActivationLayer);
+		const int32 SubLayerNum = SubLayers.Num();
+		for (int32 SubLayerIdx = 0; SubLayerIdx < SubLayerNum; SubLayerIdx++)
+		{
+			check(SubLayers[SubLayerIdx].GetOutputSize() == Encoders[SubLayerIdx].GetInputSize());
+			check(Encoders[SubLayerIdx].GetOutputSize() == OutputEncodingSize);
 		}
 
-		// Create Postfix Network
+		const TSharedPtr<Private::FAggregateOrExclusiveLayer> OrExclusiveLayer = MakeShared<Private::FAggregateOrExclusiveLayer>();
+		OrExclusiveLayer->OutputEncodingSize = OutputEncodingSize;
+		OrExclusiveLayer->SubLayerInputSizes = MakeSizesLayerInputs(SubLayers);
+		OrExclusiveLayer->SubLayerOutputSizes = MakeSizesLayerOutputs(SubLayers);
+		OrExclusiveLayer->SubLayers.Reserve(SubLayerNum);
+		OrExclusiveLayer->Encoders.Reserve(SubLayerNum);
 
-		check(PostfixLayerNum >= 1);
+		for (int32 SubLayerIdx = 0; SubLayerIdx < SubLayerNum; SubLayerIdx++)
+		{
+			OrExclusiveLayer->SubLayers.Add(SubLayers[SubLayerIdx].Layer);
+			OrExclusiveLayer->Encoders.Add(Encoders[SubLayerIdx].Layer);
+		}
 
-		TSharedRef<Private::FSequenceLayer> PostfixLayer = MakeShared<Private::FSequenceLayer>();
-		PostfixLayer->Layers.Reserve(PostfixLayerNum * 2 - 1);
+		OrExclusiveLayer->PostLoad();
 
-		for (uint32 LayerIdx = 0; LayerIdx < PostfixLayerNum; LayerIdx++)
-		{
-			const TSharedPtr<Private::FLinearLayer> LinearLayer = MakeShared<Private::FLinearLayer>();
-			LinearLayer->InputSize = HiddenUnitNum;
-			LinearLayer->OutputSize = LayerIdx + 1 == PostfixLayerNum ? OutputNum : HiddenUnitNum;
-			LinearLayer->Biases = ZerosView.Slice(0, LinearLayer->OutputSize);
-			LinearLayer->Weights = ZerosView.Slice(0, LinearLayer->InputSize * LinearLayer->OutputSize);
+		return StaticCastSharedPtr<Private::ILayer>(OrExclusiveLayer);
+	}
 
-			PostfixLayer->Layers.Emplace(LinearLayer);
+	FModelBuilderElement FModelBuilder::MakeAggregateOrInclusive(
+		const uint32 OutputEncodingSize,
+		const uint32 AttentionEncodingSize,
+		const uint32 AttentionHeadNum,
+		const TConstArrayView<FModelBuilderElement> SubLayers,
+		const TConstArrayView<FModelBuilderElement> QueryLayers,
+		const TConstArrayView<FModelBuilderElement> KeyLayers,
+		const TConstArrayView<FModelBuilderElement> ValueLayers)
+	{
+		check(SubLayers.Num() == QueryLayers.Num());
+		check(SubLayers.Num() == KeyLayers.Num());
+		check(SubLayers.Num() == ValueLayers.Num());
 
-			if (LayerIdx + 1 != PostfixLayerNum)
-			{
-				const TSharedPtr<Private::FELULayer> ActivationLayer = MakeShared<Private::FELULayer>();
-				ActivationLayer->InputOutputSize = LinearLayer->OutputSize;
+		const int32 SubLayerNum = SubLayers.Num();
+		for (int32 SubLayerIdx = 0; SubLayerIdx < SubLayerNum; SubLayerIdx++)
+		{
+			check(SubLayers[SubLayerIdx].GetOutputSize() == QueryLayers[SubLayerIdx].GetInputSize());
+			check(SubLayers[SubLayerIdx].GetOutputSize() == KeyLayers[SubLayerIdx].GetInputSize());
+			check(SubLayers[SubLayerIdx].GetOutputSize() == ValueLayers[SubLayerIdx].GetInputSize());
+			check(QueryLayers[SubLayerIdx].GetOutputSize() == AttentionHeadNum * AttentionEncodingSize);
+			check(KeyLayers[SubLayerIdx].GetOutputSize() == AttentionHeadNum * AttentionEncodingSize);
+			check(ValueLayers[SubLayerIdx].GetOutputSize() == AttentionHeadNum * OutputEncodingSize);
+		}
 
-				PostfixLayer->Layers.Emplace(ActivationLayer);
-			}
+		const TSharedPtr<Private::FAggregateOrInclusiveLayer> OrInclusiveLayer = MakeShared<Private::FAggregateOrInclusiveLayer>();
+		OrInclusiveLayer->OutputEncodingSize = OutputEncodingSize;
+		OrInclusiveLayer->AttentionEncodingSize = AttentionEncodingSize;
+		OrInclusiveLayer->AttentionHeadNum = AttentionHeadNum;
+		OrInclusiveLayer->SubLayerInputSizes = MakeSizesLayerInputs(SubLayers);
+		OrInclusiveLayer->SubLayerOutputSizes = MakeSizesLayerOutputs(SubLayers);
+		OrInclusiveLayer->SubLayers.Reserve(SubLayerNum);
+		OrInclusiveLayer->QueryLayers.Reserve(SubLayerNum);
+		OrInclusiveLayer->KeyLayers.Reserve(SubLayerNum);
+		OrInclusiveLayer->ValueLayers.Reserve(SubLayerNum);
+
+		for (int32 SubLayerIdx = 0; SubLayerIdx < SubLayerNum; SubLayerIdx++)
+		{
+			OrInclusiveLayer->SubLayers.Add(SubLayers[SubLayerIdx].Layer);
+			OrInclusiveLayer->QueryLayers.Add(QueryLayers[SubLayerIdx].Layer);
+			OrInclusiveLayer->KeyLayers.Add(KeyLayers[SubLayerIdx].Layer);
+			OrInclusiveLayer->ValueLayers.Add(ValueLayers[SubLayerIdx].Layer);
 		}
 
-		// Create Backbone
+		OrInclusiveLayer->PostLoad();
 
-		TSharedRef<Private::FMemoryBackboneLayer> BackboneLayer = MakeShared<Private::FMemoryBackboneLayer>();
-		BackboneLayer->MemorySize = MemoryNum;
-		BackboneLayer->Cell = CellLayer;
-		BackboneLayer->Prefix = PrefixLayer;
-		BackboneLayer->Postfix = PostfixLayer;
+		return StaticCastSharedPtr<Private::ILayer>(OrInclusiveLayer);
+	}
 
-		Model->Layer = BackboneLayer;
+	void FModelBuilder::Reset()
+	{
+		Rng.Reset();
+		WeightsPool.Empty();
+		SizesPool.Empty();
 	}
 
-	uint64 FMemoryBackboneModelBuilder::GetWriteByteNum() const
+	uint64 FModelBuilder::GetWriteByteNum(const FModelBuilderElement& Element) const
 	{
-		checkf(Model->Layer, TEXT("Model not built."));
+		FModelCPU Model;
+		Model.Layer = Element.Layer;
+
 		uint64 Offset = 0;
-		Model->SerializationSize(Offset);
+		Model.SerializationSize(Offset);
 		return Offset;
 	}
 
-	void FMemoryBackboneModelBuilder::WriteAndReset(TArrayView<uint8> OutBytes)
+	void FModelBuilder::WriteFileData(TArrayView<uint8> OutBytes, uint32& OutInputSize, uint32& OutOutputSize, const FModelBuilderElement& Element) const
 	{
-		checkf(Model->Layer, TEXT("Model not built."));
-		check((uint64)OutBytes.Num() >= GetWriteByteNum());
+		check((uint64)OutBytes.Num() == GetWriteByteNum(Element));
+
+		OutInputSize = Element.GetInputSize();
+		OutOutputSize = Element.GetOutputSize();
 
 		// Zero to ensure any padding due to alignment is always zero
 		FMemory::Memzero(OutBytes.GetData(), OutBytes.Num());
 
+		FModelCPU Model;
+		Model.Layer = Element.Layer;
+
 		uint64 Offset = 0;
-		Model->SerializationSave(Offset, OutBytes);
+		Model.SerializationSave(Offset, OutBytes);
+		check(Offset == OutBytes.Num());
+	}
+
+	void FModelBuilder::WriteFileData(TArray<uint8>& FileData, uint32& OutInputSize, uint32& OutOutputSize, const FModelBuilderElement& Element) const
+	{
+		FileData.SetNumUninitialized(GetWriteByteNum(Element));
+		WriteFileData(MakeArrayView(FileData), OutInputSize, OutOutputSize, Element);
+	}
+
+	void FModelBuilder::WriteFileDataAndReset(TArrayView<uint8> FileData, uint32& OutInputSize, uint32& OutOutputSize, const FModelBuilderElement& Element)
+	{
+		WriteFileData(FileData, OutInputSize, OutOutputSize, Element);
+		Reset();
+	}
+
+	void FModelBuilder::WriteFileDataAndReset(TArray<uint8>& FileData, uint32& OutInputSize, uint32& OutOutputSize, const FModelBuilderElement& Element)
+	{
+		FileData.SetNumUninitialized(GetWriteByteNum(Element));
+		WriteFileDataAndReset(MakeArrayView(FileData), OutInputSize, OutOutputSize, Element);
+	}
+
+	TArrayView<float> FModelBuilder::MakeWeightsZero(const uint32 Size)
+	{
+		TArray<float>& Values = WeightsPool.AddDefaulted_GetRef();
+		Values.Init(0.0f, Size);
+		return Values;
+	}
+
+	TArrayView<float> FModelBuilder::MakeWeightsConstant(const uint32 Size, const float Value)
+	{
+		TArray<float>& Values = WeightsPool.AddDefaulted_GetRef();
+		Values.Init(Value, Size);
+		return Values;
+	}
+
+	TArrayView<float> FModelBuilder::MakeWeightsRandomKaiming(const uint32 InputSize, const uint32 OutputSize, const float Scale)
+	{
+		TArray<float>& Values = WeightsPool.AddDefaulted_GetRef();
+		Values.SetNumUninitialized(InputSize * OutputSize);
+		
+		const float Std = Scale * FMath::Sqrt(2.0f / InputSize);
+
+		for (uint32 Idx = 0; Idx < InputSize * OutputSize; Idx++)
+		{
+			Values[Idx] = Std * Private::UniformToGaussian(Rng.FRand(), Rng.FRand());
+		}
+		
+		return Values;
+	}
+
+	TArrayView<uint32> FModelBuilder::MakeSizesZero(const uint32 Size)
+	{
+		TArray<uint32>& Values = SizesPool.AddDefaulted_GetRef();
+		Values.Init(0, Size);
+		return Values;
+	}
+
+	TArrayView<uint32> FModelBuilder::MakeSizesLayerInputs(const TConstArrayView<FModelBuilderElement> Elements)
+	{
+		const uint32 SizeNum = Elements.Num();
+		TArrayView<uint32> Sizes = MakeSizesZero(SizeNum);
+		for (uint32 SizeIdx = 0; SizeIdx < SizeNum; SizeIdx++)
+		{
+			Sizes[SizeIdx] = Elements[SizeIdx].GetInputSize();
+		}
+
+		return Sizes;
+	}
+
+	TArrayView<uint32> FModelBuilder::MakeSizesLayerOutputs(const TConstArrayView<FModelBuilderElement> Elements)
+	{
+		const uint32 SizeNum = Elements.Num();
+		TArrayView<uint32> Sizes = MakeSizesZero(SizeNum);
+		for (uint32 SizeIdx = 0; SizeIdx < SizeNum; SizeIdx++)
+		{
+			Sizes[SizeIdx] = Elements[SizeIdx].GetOutputSize();
+		}
 
-		Model->Layer.Reset();
-		ZerosData.Empty();
+		return Sizes;
 	}
 
 	//--------------------------------------------------------------------------
 
 } // namespace UE::NNE::RuntimeBasic
 
-#undef NNE_RUNTIME_BASIC_ENABLE_ISPC
+#undef NNE_RUNTIME_BASIC_TRACE_SCOPE
+#undef NNE_RUNTIME_BASIC_ENABLE_PROFILE
+#undef NNE_RUNTIME_BASIC_ENABLE_NAN_CHECK
+#undef NNE_RUNTIME_BASIC_ENABLE_ISPC
\ No newline at end of file
