LearningAgents: Large restructuring to allow for structured actions and observations.

diff --git a/Engine/Plugins/Experimental/NNERuntimeBasicCpu/Source/NNERuntimeBasicCpu/Public/NNERuntimeBasicCpuBuilder.h b/Engine/Plugins/Experimental/NNERuntimeBasicCpu/Source/NNERuntimeBasicCpu/Public/NNERuntimeBasicCpuBuilder.h
new file mode 100644
--- /dev/null
+++ b/Engine/Plugins/Experimental/NNERuntimeBasicCpu/Source/NNERuntimeBasicCpu/Public/NNERuntimeBasicCpuBuilder.h
@@ -0,0 +1,351 @@
+// Copyright Epic Games, Inc. All Rights Reserved.
+
+#pragma once
+
+#include <Math/RandomStream.h>
+
+namespace UE::NNE::RuntimeBasic
+{
+	namespace Private
+	{
+		struct ILayer;
+	}
+
+	class FModelBuilder;
+
+	/**
+	 * Represents an element of a model (such as a layer) created by the Model Builder.
+	 */
+	class NNERUNTIMEBASICCPU_API FModelBuilderElement
+	{
+		friend class FModelBuilder;
+
+	public:
+
+		FModelBuilderElement();
+		FModelBuilderElement(const TSharedPtr<Private::ILayer>& Ptr);
+		~FModelBuilderElement();
+
+		int32 GetInputSize() const;
+		int32 GetOutputSize() const;
+
+	private:
+		TSharedPtr<Private::ILayer> Layer;
+	};
+
+	/**
+	 * This class can be used to construct Models for use with the Basic CPU Runtime. 
+	 * Effectively it works by constructing the model in-memory and then serializing this
+	 * out to create a FileData object which can then be loaded with NNE.
+	 * 
+	 * Note: If you pass your own views into this model builder they must out-live the
+	 * builder itself as it will not create internal copies.
+	 */
+	class NNERUNTIMEBASICCPU_API FModelBuilder
+	{
+
+	public:
+
+		/** Common activation function types. */
+		enum class EActivationFunction : uint8
+		{
+			ReLU = 0,
+			ELU = 1,
+			TanH = 2,
+		};
+
+		/** Construct a new Model Builder with the given random seed. */
+		FModelBuilder(int32 Seed = 0x0a974e75);
+
+		/**
+		 * Makes a new linear layer.
+		 *
+		 * @param InputSize		Input Vector Size (Number of Rows)
+		 * @param OutputSize	Output Vector Size (Number of Columns)
+		 * @param Weights		Linear layer weights.
+		 * @param Biases		Linear layer biases.
+		 */
+		FModelBuilderElement MakeLinear(
+			const uint32 InputSize,
+			const uint32 OutputSize,
+			const TConstArrayView<float> Weights,
+			const TConstArrayView<float> Biases);
+
+		/**
+		 * Makes a new linear layer with randomly initialized weights using the Kaiming method.
+		 *
+		 * @param InputSize		Input Vector Size (Number of Rows)
+		 * @param OutputSize	Output Vector Size (Number of Columns)
+		 * @param WeightScale	Scaling factor for the weight creation.
+		 */
+		FModelBuilderElement MakeLinearWithRandomKaimingWeights(
+			const uint32 InputSize,
+			const uint32 OutputSize,
+			const float WeightScale = 1.0f);
+
+		/**
+		 * Makes a new multi linear layer.
+		 *
+		 * @param InputSize		Input Vector Size (Number of Rows)
+		 * @param OutputSize	Output Vector Size (Number of Columns)
+		 * @param BlockNum		Number of blocks (Number of Matrices)
+		 * @param Weights		Multi-Linear layer weights.
+		 * @param Biases		Multi-Linear layer biases.
+		 */
+		FModelBuilderElement MakeMultiLinear(
+			const uint32 InputSize,
+			const uint32 OutputSize,
+			const uint32 BlockNum,
+			const TConstArrayView<float> Weights,
+			const TConstArrayView<float> Biases);
+
+		/**
+		 * Makes a new Normalization layer.
+		 * 
+		 * @param Mean		Normalization Mean.
+		 * @param Std		Normalization Std.
+		 */
+		FModelBuilderElement MakeNormalize(
+			const uint32 InputOutputSize,
+			const TConstArrayView<float> Mean,
+			const TConstArrayView<float> Std);
+		/**
+		 * Makes a new Denormalization layer.
+		 *
+		 * @param Mean		Denormalization Mean.
+		 * @param Std		Denormalization Std.
+		 */
+		FModelBuilderElement MakeDenormalize(
+			const uint32 InputOutputSize,
+			const TConstArrayView<float> Mean,
+			const TConstArrayView<float> Std);
+
+		/** Makes a ReLU Activation Layer */
+		FModelBuilderElement MakeReLU(const uint32 InputOutputSize);
+
+		/** Makes a ELU Activation Layer */
+		FModelBuilderElement MakeELU(const uint32 InputOutputSize);
+
+		/** Makes a TanH Activation Layer */
+		FModelBuilderElement MakeTanH(const uint32 InputOutputSize);
+
+		/** Makes a Copy Layer */
+		FModelBuilderElement MakeCopy(const uint32 InputOutputSize);
+
+		/** Makes a new activation layer with the given activation function */
+		FModelBuilderElement MakeActivation(const uint32 InputOutputSize, const EActivationFunction ActivationFunction);
+
+		/**
+		 * Makes a PReLU Activation Layer
+		 *
+		 * @param LayerSize			Number of neurons in layer.
+		 * @param Alpha				PReLU alpha parameter for each neuron.
+		 */
+		FModelBuilderElement MakePReLU(const uint32 InputOutputSize, const TConstArrayView<float> Alpha);
+
+		/**
+		 * Makes a Sequence layer, which will evaluate the given list of layers in order.
+		 */
+		FModelBuilderElement MakeSequence(const TConstArrayView<FModelBuilderElement> Elements);
+
+		/**
+		 * Makes a Multi-Layer Perceptron network with randomly initialized weights using the Kaiming method.
+		 *
+		 * @param InputSize					Input Vector Size
+		 * @param OutputSize				Output Vector Size
+		 * @param HiddenSize				Number of hidden units to use on internal layers.
+		 * @param LayerNum					Number of layers. Includes input and output layers.
+		 * @param ActivationFunction		Activation function to use.
+		 * @param bActivationOnFinalLayer	If the activation function should be used on the final output layer.
+		 */
+		FModelBuilderElement MakeMLPWithRandomKaimingWeights(
+			const uint32 InputSize,
+			const uint32 OutputSize, 
+			const uint32 HiddenSize,
+			const uint32 LayerNum,
+			const EActivationFunction ActivationFunction,
+			const bool bActivationOnFinalLayer = false);
+
+		/**
+		 * Make a new Memory Cell layer.
+		 *
+		 * @param InputNum					Number of normal inputs to the model.
+		 * @param OutputNum					Number of normal outputs from the model.
+		 * @param MemoryNum					The size of the memory vector used by the model.
+		 * @param RememberLayer				Layer used for the Remember Gate.
+		 * @param PassthroughLayer			Layer used for the Passthrough Gate.
+		 * @param MemoryUpdateLayer			Layer used for updating the memory Gate.
+		 * @param OutputInputUpdateLayer	Layer used for updating the output from the input.
+		 * @param OutputMemoryUpdateLayer	Layer used for updating the output from the memory.
+		 */
+		FModelBuilderElement MakeMemoryCell(
+			const uint32 InputNum,
+			const uint32 OutputNum,
+			const uint32 MemoryNum,
+			const FModelBuilderElement& RememberLayer,
+			const FModelBuilderElement& PassthroughLayer,
+			const FModelBuilderElement& MemoryUpdateLayer,
+			const FModelBuilderElement& OutputInputUpdateLayer,
+			const FModelBuilderElement& OutputMemoryUpdateLayer);
+
+		/**
+		 * Make a new Memory Cell layer with randomly initialized weights using the Kaiming method.
+		 *
+		 * @param InputNum			Number of normal inputs to the model
+		 * @param OutputNum			Number of normal outputs from the model
+		 * @param MemoryNum			The size of the memory vector used by the model
+		 * @param WeightScale		Scaling factor for the weight creation.
+		 */
+		FModelBuilderElement MakeMemoryCellWithLinearRandomKaimingWeights(
+			const uint32 InputNum,
+			const uint32 OutputNum,
+			const uint32 MemoryNum,
+			const float WeightScale = 1.0f);
+
+		/**
+		 * Make a new Memory Backbone layer.
+		 *
+		 * @param MemoryNum				The size of the memory vector used by the model
+		 * @param Prefix				Prefix Element
+		 * @param Cell					Memory Cell Element
+		 * @param Postfix				Postfix Element
+		 */
+		FModelBuilderElement MakeMemoryBackbone(
+			const uint32 MemoryNum,
+			const FModelBuilderElement& Prefix,
+			const FModelBuilderElement& Cell,
+			const FModelBuilderElement& Postfix);
+
+		/**
+		 * Makes a Concat layer, which will evaluate each of the given elements on different slices of 
+		 * the input vector, concatenating the result into the output vector.
+		 */
+		FModelBuilderElement MakeConcat(const TConstArrayView<FModelBuilderElement> Elements);
+		
+		/**
+		 * Make a layer which runs the given sublayer on an array of elements.
+		 * 
+		 * @param ElementNum	Number of elements in the array.
+		 * @param SubLayer		Sublayer to evaluate on each item in the array.
+		 */
+		FModelBuilderElement MakeArray(const uint32 ElementNum, const FModelBuilderElement& SubLayer);
+
+		/**
+		 * Make a layer which aggregates a set of other observations using attention. This is used by LearningAgents.
+		 */
+		FModelBuilderElement MakeAggregateSet(
+			const uint32 MaxElementNum, 
+			const uint32 OutputEncodingSize,
+			const uint32 AttentionEncodingSize,
+			const uint32 AttentionHeadNum,
+			const FModelBuilderElement& SubLayer,
+			const FModelBuilderElement& QueryLayer,
+			const FModelBuilderElement& KeyLayer,
+			const FModelBuilderElement& ValueLayer);
+
+		/**
+		 * Make a layer which aggregates an Exclusive Or of other observations. This is used by LearningAgents.
+		 */
+		FModelBuilderElement MakeAggregateOrExclusive(
+			const uint32 OutputEncodingSize,
+			const TConstArrayView<FModelBuilderElement> SubLayers,
+			const TConstArrayView<FModelBuilderElement> Encoders);
+
+		/**
+		 * Make a layer which aggregates an Inclusive Or of other observations using attention. This is used by LearningAgents.
+		 */
+		FModelBuilderElement MakeAggregateOrInclusive(
+			const uint32 OutputEncodingSize,
+			const uint32 AttentionEncodingSize,
+			const uint32 AttentionHeadNum,
+			const TConstArrayView<FModelBuilderElement> SubLayers,
+			const TConstArrayView<FModelBuilderElement> QueryLayers,
+			const TConstArrayView<FModelBuilderElement> KeyLayers,
+			const TConstArrayView<FModelBuilderElement> ValueLayers);
+
+	public:
+
+		/** Creates a array of weights set to zero of the given size */
+		TArrayView<float> MakeWeightsZero(const uint32 Size);
+
+		/** Creates a array of weights set to the provided constant value of the given size */
+		TArrayView<float> MakeWeightsConstant(const uint32 Size, const float Value);
+
+		/** Creates a array of weights randomly initialized using the Kaiming method */
+		TArrayView<float> MakeWeightsRandomKaiming(const uint32 InputSize, const uint32 OutputSize, const float Scale = 1.0f);
+
+		/** Creates an array of sizes, initialized to zero. */
+		TArrayView<uint32> MakeSizesZero(const uint32 Size);
+
+		/** Creates an array of sizes from an array of builder elements' input sizes. */
+		TArrayView<uint32> MakeSizesLayerInputs(const TConstArrayView<FModelBuilderElement> Elements);
+
+		/** Creates an array of sizes from an array of builder elements' output sizes. */
+		TArrayView<uint32> MakeSizesLayerOutputs(const TConstArrayView<FModelBuilderElement> Elements);
+
+	public:
+
+		/** Reset the builder clearing all memory. */
+		void Reset();
+
+		/**
+		 * Get the number of bytes this builder currently wants to write for the given element.
+		 */
+		uint64 GetWriteByteNum(const FModelBuilderElement& Element) const;
+
+		/**
+		 * Write the Model to FileData. Use `GetWriteByteNum` to get the number of bytes this will write so that 
+		 * `FileData` can be allocated to the right size.
+		 * 
+		 * @param OutFileData			Output File data to write to
+		 * @param OutInputSize			Output for the number of float inputs to the model
+		 * @param OutOutputSize			Output for the number of float outputs from the model
+		 * @param Element				Builder element to write the model FileData for.
+		 */
+		void WriteFileData(TArrayView<uint8> OutFileData, uint32& OutInputSize, uint32& OutOutputSize, const FModelBuilderElement& Element) const;
+
+		/**
+		 * Write the model to the FileData. Use `GetWriteByteNum` to get the number of bytes this will write so that 
+		 * `FileData` can be allocated to the right size.
+		 *
+		 * @param OutFileData			Output File data to write to
+		 * @param OutInputSize			Output for the number of float inputs to the model
+		 * @param OutOutputSize			Output for the number of float outputs from the model
+		 * @param Element				Builder element to write the model FileData for.
+		 */
+		void WriteFileData(TArray<uint8>& OutFileData, uint32& OutInputSize, uint32& OutOutputSize, const FModelBuilderElement& Element) const;
+
+		/**
+		 * Write the model to the FileData and Reset the builder clearing all memory used.
+		 *
+		 * @param OutFileData			Output File data to write to
+		 * @param OutInputSize			Output for the number of float inputs to the model
+		 * @param OutOutputSize			Output for the number of float outputs from the model
+		 * @param Element				Builder element to write the model FileData for.
+		 */
+		void WriteFileDataAndReset(TArrayView<uint8> OutFileData, uint32& OutInputSize, uint32& OutOutputSize, const FModelBuilderElement& Element);
+
+		/**
+		 * Write the model to the FileData and Reset the builder clearing all memory used.
+		 * 
+		 * @param OutFileData			Output File data to write to
+		 * @param OutInputSize			Output for the number of float inputs to the model
+		 * @param OutOutputSize			Output for the number of float outputs from the model
+		 * @param Element				Builder element to write the model FileData for.
+		 */
+		void WriteFileDataAndReset(TArray<uint8>& OutFileData, uint32& OutInputSize, uint32& OutOutputSize, const FModelBuilderElement& Element);
+
+	private:
+
+		/** Random Number Stream for generating random weights */
+		FRandomStream Rng;
+
+		/** Pool of all weights data used by the `MakeWeights` functions. */
+		TArray<TArray<float>> WeightsPool;
+
+		/** Pool of all sizes data used by the `MakeSizes` functions. */
+		TArray<TArray<uint32>> SizesPool;
+	};
+
+
+}
\ No newline at end of file
