[Backout] - CL30154408

diff --git a/Engine/Plugins/Experimental/NNERuntimeBasicCpu/Source/NNERuntimeBasicCpu/Public/NNERuntimeBasicCpu.h b/Engine/Plugins/Experimental/NNERuntimeBasicCpu/Source/NNERuntimeBasicCpu/Public/NNERuntimeBasicCpu.h
new file mode 100644
--- /dev/null
+++ b/Engine/Plugins/Experimental/NNERuntimeBasicCpu/Source/NNERuntimeBasicCpu/Public/NNERuntimeBasicCpu.h
@@ -0,0 +1,165 @@
+// Copyright Epic Games, Inc. All Rights Reserved.
+
+#pragma once
+
+#include "NNERuntime.h"
+#include "NNERuntimeCPU.h"
+#include "UObject/Object.h"
+
+#include "NNERuntimeBasicCpu.generated.h"
+
+/**
+ * This plugin is a basic, performant, cross-platform CPU runtime for NNE that supports simple models such as MLPs.
+ * 
+ * To use this runtime, the custom ".ubnne" file format is used, which can be exported from python using the functions
+ * in the provided "nne_runtime_basic_cpu.py" found in the "Content" folder of this plugin. The idea behind this plugin
+ * is not to be a general purpose runtime, but rather to provide performant cross-platform implementations for simple
+ * CPU models such as MLPs with minimal overhead and memory usage.
+ */
+UCLASS()
+class NNERUNTIMEBASICCPU_API UNNERuntimeBasicCpuImpl : public UObject, public INNERuntime, public INNERuntimeCPU
+{
+	GENERATED_BODY()
+
+public:
+
+	virtual FString GetRuntimeName() const override { return TEXT("NNERuntimeBasicCpu"); };
+
+	virtual bool CanCreateModelData(FString FileType, TConstArrayView<uint8> FileData, FGuid FileId, const ITargetPlatform* TargetPlatform) const override;
+	virtual TSharedPtr<UE::NNE::FSharedModelData> CreateModelData(FString FileType, TConstArrayView<uint8> FileData, FGuid FileId, const ITargetPlatform* TargetPlatform) override;
+	virtual FString GetModelDataIdentifier(FString FileType, TConstArrayView<uint8> FileData, FGuid FileId, const ITargetPlatform* TargetPlatform) override;
+
+	virtual bool CanCreateModelCPU(TObjectPtr<UNNEModelData> ModelData) const override;
+	virtual TSharedPtr<UE::NNE::IModelCPU> CreateModelCPU(TObjectPtr<UNNEModelData> ModelData) override;
+
+private:
+
+	static const uint32 Alignment;
+};
+
+namespace UE::NNE::RuntimeBasic
+{
+	class FModelCPU;
+
+	/**
+	 * This class can be used to construct file data for basic sequential models at runtime. It should not live longer
+	 * than the views of data passed in by the various functions for adding layers.
+	 */
+	class NNERUNTIMEBASICCPU_API FSequentialModelBuilder
+	{
+
+	public:
+
+		FSequentialModelBuilder();
+
+		/**
+		 * Adds a new linear layer to the model.
+		 *
+		 * @param InputSize		Input Vector Size (Number of Rows)
+		 * @param OutputSize	Output Vector Size (Number of Columns)
+		 * @param Weights		Linear layer weights.
+		 * @param Biases		Linear layer biases.
+		 */
+		void AddLinear(
+			uint32 InputSize,
+			uint32 OutputSize,
+			TConstArrayView<float> Weights,
+			TConstArrayView<float> Biases);
+
+		/**
+		 * Adds a new multi linear layer to the model.
+		 *
+		 * @param InputSize		Input Vector Size (Number of Rows)
+		 * @param OutputSize	Output Vector Size (Number of Columns)
+		 * @param BlockNum		Number of blocks (Number of Matrices)
+		 * @param Weights		Multi-Linear layer weights.
+		 * @param Biases		Multi-Linear layer biases.
+		 */
+		void AddMultiLinear(
+			uint32 InputSize,
+			uint32 OutputSize,
+			uint32 BlockNum,
+			TConstArrayView<float> Weights,
+			TConstArrayView<float> Biases);
+
+		/** Add a ReLU Activation Layer */
+		void AddReLU();
+
+		/** Add a ELU Activation Layer */
+		void AddELU();
+
+		/** Add a TanH Activation Layer */
+		void AddTanH();
+
+		/**
+		 * Add a PReLU Activation Layer
+		 *
+		 * @param Alpha				PReLU alpha parameter for each neuron.
+		 */
+		void AddPReLU(TConstArrayView<float> Alpha);
+
+		/**
+		 * Get the number of bytes this builder currently wants to write.
+		 */
+		uint64 GetWriteByteNum() const;
+
+		/**
+		 * Write the Model to FileData and reset this builder, freeing all the memory used. Use `GetWriteByteNum` to
+		 * get the number of bytes this will write so that `FileData` can be allocated to the right size.
+		 */
+		void WriteAndReset(TArrayView<uint8> FileData);
+
+	private:
+
+		TSharedPtr<FModelCPU> Model;
+	};
+
+	/**
+	 * This class can be used to construct file data for "memory backbone" model which contains a recurrent memory 
+	 * cell sandwiched between two feedforward neural networks. In this case it is only possible to build an empty
+	 * model with zero'd weights.
+	 */
+	class NNERUNTIMEBASICCPU_API FMemoryBackboneModelBuilder
+	{
+
+	public:
+
+		FMemoryBackboneModelBuilder();
+
+		/**
+		 * Builds an empty Memory Backbone model.
+		 *
+		 * @param InputNum			Number of normal inputs to the model
+		 * @param OutputNum			Number of normal outputs from the model
+		 * @param MemoryNum			The size of the memory vector used by the model
+		 * @param HiddenUnitNum		The number of hidden units used in all internal layers.
+		 * @param PrefixLayerNum	The number of layers used in the prefix feed-forward neural network
+		 * @param PostfixLayerNum	The number of layers used in the postfix feed-forward neural network
+		 */
+		void BuildEmptyModel(
+			const uint32 InputNum,
+			const uint32 OutputNum,
+			const uint32 MemoryNum,
+			const uint32 HiddenUnitNum,
+			const uint32 PrefixLayerNum,
+			const uint32 PostfixLayerNum);
+
+		/**
+		 * Get the number of bytes this builder currently wants to write.
+		 */
+		uint64 GetWriteByteNum() const;
+
+		/**
+		 * Write the Model to FileData and reset this builder, freeing all the memory used. Use `GetWriteByteNum` to
+		 * get the number of bytes this will write so that `FileData` can be allocated to the right size.
+		 */
+		void WriteAndReset(TArrayView<uint8> FileData);
+
+	private:
+
+		TSharedPtr<FModelCPU> Model;
+		TArray<float> ZerosData;
+	};
+
+
+}
\ No newline at end of file
