NNERuntimeBasicCpu: Added Clamp Op

diff --git a/Engine/Plugins/Experimental/NNERuntimeBasicCpu/Source/NNERuntimeBasicCpu/Private/NNERuntimeBasicCpu.ispc b/Engine/Plugins/Experimental/NNERuntimeBasicCpu/Source/NNERuntimeBasicCpu/Private/NNERuntimeBasicCpu.ispc
--- a/Engine/Plugins/Experimental/NNERuntimeBasicCpu/Source/NNERuntimeBasicCpu/Private/NNERuntimeBasicCpu.ispc
+++ b/Engine/Plugins/Experimental/NNERuntimeBasicCpu/Source/NNERuntimeBasicCpu/Private/NNERuntimeBasicCpu.ispc
@@ -43,6 +43,22 @@ export void NNERuntimeBasicCPUOperatorDenormalize(
 	}
 }
 
+export void NNERuntimeBasicCPUOperatorClamp(
+	uniform float Output[],
+	uniform const float Input[],
+	uniform const float MinValues[],
+	uniform const float MaxValues[],
+	uniform const int BatchNum,
+	uniform const int InputOutputNum,
+	uniform const int OutputStride,
+	uniform const int InputStride)
+{
+	foreach(BatchIdx = 0 ... BatchNum, Idx = 0 ... InputOutputNum)
+	{
+		Output[BatchIdx * OutputStride + Idx] = clamp(Input[BatchIdx * InputStride + Idx], MinValues[Idx], MaxValues[Idx]);
+	}
+}
+
 export void NNERuntimeBasicCPUOperatorLinear(
 	uniform float Output[],
 	uniform const float Input[],
diff --git a/Engine/Plugins/Experimental/NNERuntimeBasicCpu/Source/NNERuntimeBasicCpu/Private/NNERuntimeBasicCpuModel.cpp b/Engine/Plugins/Experimental/NNERuntimeBasicCpu/Source/NNERuntimeBasicCpu/Private/NNERuntimeBasicCpuModel.cpp
--- a/Engine/Plugins/Experimental/NNERuntimeBasicCpu/Source/NNERuntimeBasicCpu/Private/NNERuntimeBasicCpuModel.cpp
+++ b/Engine/Plugins/Experimental/NNERuntimeBasicCpu/Source/NNERuntimeBasicCpu/Private/NNERuntimeBasicCpuModel.cpp
@@ -288,6 +288,39 @@ namespace UE::NNE::RuntimeBasic
 #endif
 		}
 
+		static inline void OperatorClamp(
+			float* RESTRICT Output,
+			const float* RESTRICT Input,
+			const float* RESTRICT MinValues,
+			const float* RESTRICT MaxValues,
+			const uint32 BatchSize,
+			const uint32 InputOutputSize,
+			const uint32 OutputStride,
+			const uint32 InputStride)
+		{
+			NNE_RUNTIME_BASIC_TRACE_SCOPE(NNE::RuntimeBasic::Private::OperatorClamp);
+
+#if NNE_RUNTIME_BASIC_ENABLE_ISPC
+			ispc::NNERuntimeBasicCPUOperatorClamp(
+				Output,
+				Input,
+				MinValues,
+				MaxValues,
+				BatchSize,
+				InputOutputSize,
+				OutputStride,
+				InputStride);
+#else
+			for (uint32 BatchIdx = 0; BatchIdx < BatchSize; BatchIdx++)
+			{
+				for (uint32 Idx = 0; Idx < InputOutputSize; Idx++)
+				{
+					Output[BatchIdx * OutputStride + Idx] = FMath::Clamp(Input[BatchIdx * InputStride + Idx], MinValues[Idx], MaxValues[Idx]);
+				}
+			}
+#endif
+		}
+
 		static inline void OperatorLinear(
 			float* RESTRICT Output,
 			const float* RESTRICT Input,
@@ -1143,6 +1176,8 @@ namespace UE::NNE::RuntimeBasic
 			AggregateSet = 15,
 			AggregateOrExclusive = 16,
 			AggregateOrInclusive = 17,
+
+			Clamp = 18,
 		};
 
 		//--------------------------------------------------------------------------
@@ -3302,6 +3337,69 @@ namespace UE::NNE::RuntimeBasic
 			ValueBuffer.SetNumUninitialized(MaxBatchSize * SubLayerNum * AggregateOrInclusiveLayer.AttentionHeadNum * AggregateOrInclusiveLayer.OutputEncodingSize);
 		}
 
+
+		//--------------------------------------------------------------------------
+
+		struct FClampLayer : public ILayer
+		{
+			virtual ELayerType GetLayerType() const override final { return ELayerType::Clamp; }
+			virtual uint32 GetInputSize() const override final { return InputOutputSize; }
+			virtual uint32 GetOutputSize() const override final { return InputOutputSize; }
+
+			virtual void SerializationSize(uint64& InOutOffset) const override final
+			{
+				Serialization::Size(InOutOffset, InputOutputSize);
+				Serialization::Size(InOutOffset, MinValues);
+				Serialization::Size(InOutOffset, MaxValues);
+			}
+
+			virtual void SerializationLoad(uint64& InOutOffset, TConstArrayView<uint8> Data) override final
+			{
+				Serialization::Load(InOutOffset, InputOutputSize, Data);
+				Serialization::Load(InOutOffset, MinValues, Data, InputOutputSize);
+				Serialization::Load(InOutOffset, MaxValues, Data, InputOutputSize);
+			}
+
+			virtual void SerializationSave(uint64& InOutOffset, TArrayView<uint8> Data) const override final
+			{
+				Serialization::Save(InOutOffset, InputOutputSize, Data);
+				Serialization::Save(InOutOffset, MinValues, Data);
+				Serialization::Save(InOutOffset, MaxValues, Data);
+			}
+
+			virtual void Evaluate(
+				ILayerInstance* Instance,
+				float* OutputBuffer,
+				const float* InputBuffer,
+				const uint32 BatchSize,
+				const uint32 OutputBufferSize,
+				const uint32 InputBufferSize,
+				const uint32 OutputBufferStride,
+				const uint32 InputBufferStride) override final
+			{
+				NNE_RUNTIME_BASIC_TRACE_SCOPE(NNE::RuntimeBasic::Private::FClampLayer::Evaluate);
+				check(OutputBufferSize == GetOutputSize() && InputBufferSize == GetInputSize());
+				check(Instance == nullptr);
+				OperatorNanCheck(InputBuffer, BatchSize, InputBufferSize, InputBufferStride);
+
+				OperatorClamp(
+					OutputBuffer,
+					InputBuffer,
+					MinValues.GetData(),
+					MaxValues.GetData(),
+					BatchSize,
+					InputOutputSize,
+					OutputBufferStride,
+					InputBufferStride);
+
+				OperatorNanCheck(OutputBuffer, BatchSize, OutputBufferSize, OutputBufferStride);
+			}
+
+			uint32 InputOutputSize = 0;
+			TConstArrayView<float> MinValues;
+			TConstArrayView<float> MaxValues;
+		};
+
 		//--------------------------------------------------------------------------
 		// Layer Serialization
 		//--------------------------------------------------------------------------
@@ -3347,6 +3445,7 @@ namespace UE::NNE::RuntimeBasic
 					case ELayerType::AggregateSet: OutLayer = MakeShared<FAggregateSetLayer>(); break;
 					case ELayerType::AggregateOrExclusive: OutLayer = MakeShared<FAggregateOrExclusiveLayer>(); break;
 					case ELayerType::AggregateOrInclusive: OutLayer = MakeShared<FAggregateOrInclusiveLayer>(); break;
+					case ELayerType::Clamp: OutLayer = MakeShared<FClampLayer>(); break;
 					default: checkf(false, TEXT("Unknown Layer Id %i"), LayerTypeId);
 					}
 				}
@@ -3609,6 +3708,9 @@ namespace UE::NNE::RuntimeBasic
 		const TConstArrayView<float> Mean,
 		const TConstArrayView<float> Std)
 	{
+		check(Mean.Num() == InputOutputSize);
+		check(Std.Num() == InputOutputSize);
+
 		const TSharedPtr<Private::FNormalizeLayer> NormalizeLayer = MakeShared<Private::FNormalizeLayer>();
 		NormalizeLayer->InputOutputSize = InputOutputSize;
 		NormalizeLayer->Mean = Mean;
@@ -3622,6 +3724,9 @@ namespace UE::NNE::RuntimeBasic
 		const TConstArrayView<float> Mean,
 		const TConstArrayView<float> Std)
 	{
+		check(Mean.Num() == InputOutputSize);
+		check(Std.Num() == InputOutputSize);
+
 		const TSharedPtr<Private::FDenormalizeLayer> DenormalizeLayer = MakeShared<Private::FDenormalizeLayer>();
 		DenormalizeLayer->InputOutputSize = InputOutputSize;
 		DenormalizeLayer->Mean = Mean;
@@ -3658,6 +3763,18 @@ namespace UE::NNE::RuntimeBasic
 		return StaticCastSharedPtr<Private::ILayer>(Layer);
 	}
 
+	FModelBuilderElement FModelBuilder::MakeClamp(const uint32 InputOutputSize, const TConstArrayView<float> MinValues, const TConstArrayView<float> MaxValues)
+	{
+		check(MinValues.Num() == InputOutputSize);
+		check(MaxValues.Num() == InputOutputSize);
+
+		const TSharedPtr<Private::FClampLayer> Layer = MakeShared<Private::FClampLayer>();
+		Layer->InputOutputSize = InputOutputSize;
+		Layer->MinValues = MinValues;
+		Layer->MaxValues = MaxValues;
+		return StaticCastSharedPtr<Private::ILayer>(Layer);
+	}
+
 	FModelBuilderElement FModelBuilder::MakeActivation(const uint32 InputOutputSize, const EActivationFunction ActivationFunction)
 	{
 		switch (ActivationFunction)
diff --git a/Engine/Plugins/Experimental/NNERuntimeBasicCpu/Source/NNERuntimeBasicCpu/Public/NNERuntimeBasicCpuBuilder.h b/Engine/Plugins/Experimental/NNERuntimeBasicCpu/Source/NNERuntimeBasicCpu/Public/NNERuntimeBasicCpuBuilder.h
--- a/Engine/Plugins/Experimental/NNERuntimeBasicCpu/Source/NNERuntimeBasicCpu/Public/NNERuntimeBasicCpuBuilder.h
+++ b/Engine/Plugins/Experimental/NNERuntimeBasicCpu/Source/NNERuntimeBasicCpu/Public/NNERuntimeBasicCpuBuilder.h
@@ -132,6 +132,9 @@ namespace UE::NNE::RuntimeBasic
 		/** Makes a Copy Layer */
 		FModelBuilderElement MakeCopy(const uint32 InputOutputSize);
 
+		/** Makes a Clamp Layer */
+		FModelBuilderElement MakeClamp(const uint32 InputOutputSize, const TConstArrayView<float> MinValues, const TConstArrayView<float> MaxValues);
+
 		/** Makes a new activation layer with the given activation function */
 		FModelBuilderElement MakeActivation(const uint32 InputOutputSize, const EActivationFunction ActivationFunction);
 
