LearningAgents: Changed how the number of elements is encoded in a Set observation.

#virtualized

diff --git a/Engine/Plugins/Experimental/LearningAgents/Source/Learning/Private/LearningObservation.cpp b/Engine/Plugins/Experimental/LearningAgents/Source/Learning/Private/LearningObservation.cpp
--- a/Engine/Plugins/Experimental/LearningAgents/Source/Learning/Private/LearningObservation.cpp
+++ b/Engine/Plugins/Experimental/LearningAgents/Source/Learning/Private/LearningObservation.cpp
@@ -199,7 +199,7 @@ namespace UE::Learning::Observation
 		const int32 Index = Types.Add(EType::Set);
 		Names.Add(Name);
 		ObservationVectorSizes.Add(GetObservationVectorSize(Parameters.Element) * Parameters.MaxNum + Parameters.MaxNum);
-		EncodedVectorSizes.Add(Parameters.ValueEncodingSize * Parameters.AttentionHeadNum + Parameters.MaxNum);
+		EncodedVectorSizes.Add(Parameters.ValueEncodingSize * Parameters.AttentionHeadNum + 1);
 		TypeDataIndices.Add(SetData.Add(ElementData));
 
 		return { Index, Generation };
diff --git a/Engine/Plugins/Experimental/LearningAgents/Source/Learning/Public/LearningAction.h b/Engine/Plugins/Experimental/LearningAgents/Source/Learning/Public/LearningAction.h
--- a/Engine/Plugins/Experimental/LearningAgents/Source/Learning/Public/LearningAction.h
+++ b/Engine/Plugins/Experimental/LearningAgents/Source/Learning/Public/LearningAction.h
@@ -392,7 +392,7 @@ namespace UE::Learning::Action
 
 	/**
 	 * Gets a hash value representing object compatibility between schemas i.e. if objects from one schema can be used by objects expecting another 
-	 * schema. This is not a crytographic hash, and so `AreSchemaObjectsCompatible` should still be used as the ultimate source of truth.
+	 * schema. This is not a cryptographic hash, and so `AreSchemaObjectsCompatible` should still be used as the ultimate source of truth.
 	 * This function returns an int32 so that it can be used in blueprints.
 	 * 
 	 * @param Schema					Action Schema
diff --git a/Engine/Plugins/Experimental/LearningAgents/Source/Learning/Public/LearningObservation.h b/Engine/Plugins/Experimental/LearningAgents/Source/Learning/Public/LearningObservation.h
--- a/Engine/Plugins/Experimental/LearningAgents/Source/Learning/Public/LearningObservation.h
+++ b/Engine/Plugins/Experimental/LearningAgents/Source/Learning/Public/LearningObservation.h
@@ -374,7 +374,7 @@ namespace UE::Learning::Observation
 
 	/**
 	 * Gets a hash value representing object compatibility between schemas i.e. if objects from one schema can be used by objects expecting another
-	 * schema. This is not a crytographic hash, and so `AreSchemaObjectsCompatible` should still be used as the ultimate source of truth.
+	 * schema. This is not a cryptographic hash, and so `AreSchemaObjectsCompatible` should still be used as the ultimate source of truth.
 	 * This function returns an int32 so that it can be used in blueprints.
 	 *
 	 * @param Schema					Observation Schema
diff --git a/Engine/Plugins/Experimental/NNERuntimeBasicCpu/Source/NNERuntimeBasicCpu/Private/NNERuntimeBasicCpuModel.cpp b/Engine/Plugins/Experimental/NNERuntimeBasicCpu/Source/NNERuntimeBasicCpu/Private/NNERuntimeBasicCpuModel.cpp
--- a/Engine/Plugins/Experimental/NNERuntimeBasicCpu/Source/NNERuntimeBasicCpu/Private/NNERuntimeBasicCpuModel.cpp
+++ b/Engine/Plugins/Experimental/NNERuntimeBasicCpu/Source/NNERuntimeBasicCpu/Private/NNERuntimeBasicCpuModel.cpp
@@ -902,6 +902,19 @@ namespace UE::NNE::RuntimeBasic
 #endif
 		}
 
+		static inline void OperatorEncodeElementNums(
+			float* RESTRICT OutputBuffer,
+			const uint32* RESTRICT ElementNums,
+			const uint32 MaxElementNum,
+			const uint32 BatchSize,
+			const uint32 OutputStride)
+		{
+			for (uint32 BatchIdx = 0; BatchIdx < BatchSize; BatchIdx++)
+			{
+				OutputBuffer[BatchIdx * OutputStride] = (float)ElementNums[BatchIdx] / (float)MaxElementNum;
+			}
+		}
+
 		static inline void OperatorAggregateSoftmaxPlusOneInplace(
 			float* RESTRICT AttentionMaxs,
 			float* RESTRICT AttentionDenoms,
@@ -2573,7 +2586,7 @@ namespace UE::NNE::RuntimeBasic
 			virtual TSharedPtr<ILayerInstance> MakeInstance() const { return MakeShared<FAggregateSetLayerInstance>(*this); };
 			virtual ELayerType GetLayerType() const override final { return ELayerType::AggregateSet; }
 			virtual uint32 GetInputSize() const override final { return MaxElementNum * ElementInputSize + MaxElementNum; }
-			virtual uint32 GetOutputSize() const override final { return AttentionHeadNum * OutputEncodingSize + MaxElementNum; }
+			virtual uint32 GetOutputSize() const override final { return AttentionHeadNum * OutputEncodingSize + 1; }
 
 			virtual void SerializationSize(uint64& InOutOffset) const override final
 			{
@@ -2735,15 +2748,14 @@ namespace UE::NNE::RuntimeBasic
 					AttentionHeadNum,
 					OutputBufferStride);
 
-				// Append Element Mask
+				// Append Element Nums
 
-				OperatorCopy(
+				OperatorEncodeElementNums(
 					OutputBuffer + AttentionHeadNum * OutputEncodingSize,
-					InputBuffer + MaxElementNum * ElementInputSize,
-					BatchSize,
+					AggregateSetInstance->ElementNums.GetData(),
 					MaxElementNum,
-					OutputBufferStride,
-					InputBufferStride);
+					BatchSize,
+					OutputBufferStride);
 
 				OperatorNanCheck(OutputBuffer, BatchSize, OutputBufferSize, OutputBufferStride);
 			}
