NNERuntimeBasicCPU: Added SparseMixtureOfExperts Layer and some better support for creating CompressedLinear layers on both the C++ and Python side.

diff --git a/Engine/Plugins/Experimental/NNERuntimeBasicCpu/Source/NNERuntimeBasicCpu/Private/NNERuntimeBasicCpuModel.cpp b/Engine/Plugins/Experimental/NNERuntimeBasicCpu/Source/NNERuntimeBasicCpu/Private/NNERuntimeBasicCpuModel.cpp
--- a/Engine/Plugins/Experimental/NNERuntimeBasicCpu/Source/NNERuntimeBasicCpu/Private/NNERuntimeBasicCpuModel.cpp
+++ b/Engine/Plugins/Experimental/NNERuntimeBasicCpu/Source/NNERuntimeBasicCpu/Private/NNERuntimeBasicCpuModel.cpp
@@ -1158,6 +1158,105 @@ namespace UE::NNE::RuntimeBasic
 			}
 		}
 
+		static inline void OperatorGatherTopTwoSubLayerBatchIndices(
+			TArrayView<TArray<uint32>> SubLayerBatchIndices,
+			uint32* RESTRICT BatchSubLayerIndex0,
+			uint32* RESTRICT BatchSubLayerIndex1,
+			float* RESTRICT BatchSubLayerWeight0,
+			float* RESTRICT BatchSubLayerWeight1,
+			uint32* RESTRICT BatchSubLayerOutputIndex0,
+			uint32* RESTRICT BatchSubLayerOutputIndex1,
+			const float* RESTRICT SubLayerGateBuffer,
+			const uint32 BatchSize,
+			const uint32 SubLayerGateSize,
+			const uint32 SubLayerGateStride)
+		{
+			NNE_RUNTIME_BASIC_TRACE_SCOPE(NNE::RuntimeBasic::Private::OperatorGatherTopTwoSubLayerBatchIndices);
+
+			check(SubLayerGateSize >= 2);
+
+			for (uint32 SubLayerIdx = 0; SubLayerIdx < SubLayerGateSize; SubLayerIdx++)
+			{
+				SubLayerBatchIndices[SubLayerIdx].Reset();
+			}
+
+			for (uint32 BatchIdx = 0; BatchIdx < BatchSize; BatchIdx++)
+			{
+				int32 BestIdx0 = INDEX_NONE;
+				int32 BestIdx1 = INDEX_NONE;
+				float BestVal0 = -FLT_MAX;
+				float BestVal1 = -FLT_MAX;
+
+				for (uint32 SubLayerIdx = 0; SubLayerIdx < SubLayerGateSize; SubLayerIdx++)
+				{
+					const float GateValue = SubLayerGateBuffer[BatchIdx * SubLayerGateStride + SubLayerIdx];
+
+					if (GateValue > BestVal0)
+					{
+						BestIdx1 = BestIdx0;
+						BestVal1 = BestVal0;
+						BestIdx0 = SubLayerIdx;
+						BestVal0 = GateValue;
+						continue;
+					}
+
+					if (GateValue > BestVal1)
+					{
+						BestIdx1 = SubLayerIdx;
+						BestVal1 = GateValue;
+						continue;
+					}
+				}
+
+				const float ExpVal0 = FMath::Exp(BestVal0 - FMath::Max(BestVal0, BestVal1));
+				const float ExpVal1 = FMath::Exp(BestVal1 - FMath::Max(BestVal0, BestVal1));
+
+				BatchSubLayerIndex0[BatchIdx] = BestIdx0;
+				BatchSubLayerIndex1[BatchIdx] = BestIdx1;
+				BatchSubLayerWeight0[BatchIdx] = ExpVal0 / (ExpVal0 + ExpVal1);
+				BatchSubLayerWeight1[BatchIdx] = ExpVal1 / (ExpVal0 + ExpVal1);
+				BatchSubLayerOutputIndex0[BatchIdx] = SubLayerBatchIndices[BestIdx0].Add(BatchIdx);
+				BatchSubLayerOutputIndex1[BatchIdx] = SubLayerBatchIndices[BestIdx1].Add(BatchIdx);
+			}
+		}
+
+		static inline void OperatorGatherTopTwoFromSubLayers(
+			float* RESTRICT OutputBuffer,
+			const uint32* RESTRICT BatchSubLayerIndex0,
+			const uint32* RESTRICT BatchSubLayerIndex1,
+			const float* RESTRICT BatchSubLayerWeight0,
+			const float* RESTRICT BatchSubLayerWeight1,
+			const uint32* RESTRICT BatchSubLayerOutputIndex0,
+			const uint32* RESTRICT BatchSubLayerOutputIndex1,
+			const TConstArrayView<TArray<float>> SubLayerOutputBuffer,
+			const uint32 BatchSize,
+			const uint32 OutputBufferSize,
+			const uint32 SubLayerOutputStride,
+			const uint32 OutputBufferStride)
+		{
+			NNE_RUNTIME_BASIC_TRACE_SCOPE(NNE::RuntimeBasic::Private::OperatorGatherTopTwoFromSubLayers);
+
+			for (uint32 BatchIdx = 0; BatchIdx < BatchSize; BatchIdx++)
+			{
+				const uint32 SubLayerIndex0 = BatchSubLayerIndex0[BatchIdx];
+				const uint32 SubLayerIndex1 = BatchSubLayerIndex1[BatchIdx];
+				const float Weight0 = BatchSubLayerWeight0[BatchIdx];
+				const float Weight1 = BatchSubLayerWeight1[BatchIdx];
+				const uint32 SubLayerOutputIndex0 = BatchSubLayerOutputIndex0[BatchIdx];
+				const uint32 SubLayerOutputIndex1 = BatchSubLayerOutputIndex1[BatchIdx];
+
+				const float* RESTRICT SubLayerOutputBuffer0 = SubLayerOutputBuffer[SubLayerIndex0].GetData();
+				const float* RESTRICT SubLayerOutputBuffer1 = SubLayerOutputBuffer[SubLayerIndex1].GetData();
+
+				for (uint32 OutputIdx = 0; OutputIdx < OutputBufferSize; OutputIdx++)
+				{
+					OutputBuffer[BatchIdx * OutputBufferStride + OutputIdx] =
+						Weight0 * SubLayerOutputBuffer0[SubLayerOutputIndex0 * SubLayerOutputStride + OutputIdx] +
+						Weight1 * SubLayerOutputBuffer1[SubLayerOutputIndex1 * SubLayerOutputStride + OutputIdx];
+				}
+			}
+		}
+
 		//--------------------------------------------------------------------------
 		// Layer Types
 		//--------------------------------------------------------------------------
@@ -1191,6 +1290,7 @@ namespace UE::NNE::RuntimeBasic
 			AggregateOrInclusive = 17,
 
 			Clamp = 18,
+			SparseMixtureOfExperts = 19,
 		};
 
 		//--------------------------------------------------------------------------
@@ -3412,6 +3512,217 @@ namespace UE::NNE::RuntimeBasic
 			TConstArrayView<float> MaxValues;
 		};
 
+		//--------------------------------------------------------------------------
+
+		struct FSparseMixtureOfExpertsLayer;
+
+		struct FSparseMixtureOfExpertsLayerInstance : public ILayerInstance
+		{
+			FSparseMixtureOfExpertsLayerInstance(const FSparseMixtureOfExpertsLayer& InSparseMixtureOfExpertsLayer);
+
+			void SetMaxBatchSize(const uint32 MaxBatchSize) override final;
+
+			const FSparseMixtureOfExpertsLayer& SparseMixtureOfExpertsLayer;
+			TSharedPtr<ILayerInstance> GatingInstance;
+			TArray<float> GatingOutputBuffer;
+
+			TArray<TSharedPtr<ILayerInstance>, TInlineAllocator<32>> SubLayerInstances;
+			TArray<uint32> BatchSubLayerIndex0;
+			TArray<uint32> BatchSubLayerIndex1;
+			TArray<float> BatchSubLayerWeight0;
+			TArray<float> BatchSubLayerWeight1;
+			TArray<uint32> BatchSubLayerOutputIndex0;
+			TArray<uint32> BatchSubLayerOutputIndex1;
+			TArray<TArray<uint32>, TInlineAllocator<32>> SubLayerBatchIndices;
+			TArray<TArray<float>, TInlineAllocator<32>> SubLayerInputBuffers;
+			TArray<TArray<float>, TInlineAllocator<32>> SubLayerOutputBuffers;
+		};
+
+		struct FSparseMixtureOfExpertsLayer : public ILayer
+		{
+			virtual TSharedPtr<ILayerInstance> MakeInstance() const { return MakeShared<FSparseMixtureOfExpertsLayerInstance>(*this); };
+			virtual ELayerType GetLayerType() const override final { return ELayerType::SparseMixtureOfExperts; }
+			virtual uint32 GetInputSize() const override final { return InputSize; }
+			virtual uint32 GetOutputSize() const override final { return OutputSize; }
+
+			virtual void SerializationSize(uint64& InOutOffset) const override final
+			{
+				Serialization::Size(InOutOffset, InputSize);
+				Serialization::Size(InOutOffset, OutputSize);
+				Serialization::Size(InOutOffset, GatingLayer);
+				Serialization::Size(InOutOffset, (uint32)SubLayers.Num());
+				Serialization::Size(InOutOffset, SubLayers);
+			}
+
+			virtual void SerializationLoad(uint64& InOutOffset, TConstArrayView<uint8> Data) override final
+			{
+				Serialization::Load(InOutOffset, InputSize, Data);
+				Serialization::Load(InOutOffset, OutputSize, Data);
+				Serialization::Load(InOutOffset, GatingLayer, Data);
+				uint32 SubLayerNum = 0;
+				Serialization::Load(InOutOffset, SubLayerNum, Data);
+				SubLayers.Init(nullptr, SubLayerNum);
+				Serialization::Load(InOutOffset, SubLayers, Data);
+			}
+
+			virtual void SerializationSave(uint64& InOutOffset, TArrayView<uint8> Data) const override final
+			{
+				Serialization::Save(InOutOffset, InputSize, Data);
+				Serialization::Save(InOutOffset, OutputSize, Data);
+				Serialization::Save(InOutOffset, GatingLayer, Data);
+				Serialization::Save(InOutOffset, (uint32)SubLayers.Num(), Data);
+				Serialization::Save(InOutOffset, SubLayers, Data);
+			}
+
+			virtual void Evaluate(
+				ILayerInstance* Instance,
+				float* OutputBuffer,
+				const float* InputBuffer,
+				const uint32 BatchSize,
+				const uint32 OutputBufferSize,
+				const uint32 InputBufferSize,
+				const uint32 OutputBufferStride,
+				const uint32 InputBufferStride) override final
+			{
+				NNE_RUNTIME_BASIC_TRACE_SCOPE(NNE::RuntimeBasic::Private::FSparseMixtureOfExpertsLayer::Evaluate);
+				check(OutputBufferSize == GetOutputSize() && InputBufferSize == GetInputSize());
+
+				FSparseMixtureOfExpertsLayerInstance* SparseMixtureOfExpertsInstance = StaticCast<FSparseMixtureOfExpertsLayerInstance*>(Instance);
+				check(SparseMixtureOfExpertsInstance);
+
+				OperatorNanCheck(InputBuffer, BatchSize, InputBufferSize, InputBufferStride);
+
+				const uint32 SubLayerNum = SubLayers.Num();
+
+				// Evaluate Gating Layer
+
+				GatingLayer->Evaluate(
+					SparseMixtureOfExpertsInstance->GatingInstance.Get(),
+					SparseMixtureOfExpertsInstance->GatingOutputBuffer.GetData(),
+					InputBuffer,
+					BatchSize,
+					SubLayerNum,
+					InputBufferSize,
+					SubLayerNum,
+					InputBufferStride);
+
+				// Gather Batch SubLayer Indices according to Top-2 Experts
+
+				OperatorGatherTopTwoSubLayerBatchIndices(
+					SparseMixtureOfExpertsInstance->SubLayerBatchIndices,
+					SparseMixtureOfExpertsInstance->BatchSubLayerIndex0.GetData(),
+					SparseMixtureOfExpertsInstance->BatchSubLayerIndex1.GetData(),
+					SparseMixtureOfExpertsInstance->BatchSubLayerWeight0.GetData(),
+					SparseMixtureOfExpertsInstance->BatchSubLayerWeight1.GetData(),
+					SparseMixtureOfExpertsInstance->BatchSubLayerOutputIndex0.GetData(),
+					SparseMixtureOfExpertsInstance->BatchSubLayerOutputIndex1.GetData(),
+					SparseMixtureOfExpertsInstance->GatingOutputBuffer.GetData(),
+					BatchSize,
+					SubLayerNum,
+					SubLayerNum);
+
+				// Evaluate Each Sublayer on the associated batch items
+
+				for (uint32 SubLayerIdx = 0; SubLayerIdx < SubLayerNum; SubLayerIdx++)
+				{
+					const uint32 SubLayerBatchSize = SparseMixtureOfExpertsInstance->SubLayerBatchIndices[SubLayerIdx].Num();
+
+					if (SubLayerBatchSize == 0) { continue; }
+
+					OperatorGather(
+						SparseMixtureOfExpertsInstance->SubLayerInputBuffers[SubLayerIdx].GetData(),
+						InputBuffer,
+						SparseMixtureOfExpertsInstance->SubLayerBatchIndices[SubLayerIdx].GetData(),
+						SubLayerBatchSize,
+						InputSize,
+						InputSize,
+						InputBufferStride);
+
+					SubLayers[SubLayerIdx]->Evaluate(
+						SparseMixtureOfExpertsInstance->SubLayerInstances[SubLayerIdx].Get(),
+						SparseMixtureOfExpertsInstance->SubLayerOutputBuffers[SubLayerIdx].GetData(),
+						SparseMixtureOfExpertsInstance->SubLayerInputBuffers[SubLayerIdx].GetData(),
+						SubLayerBatchSize,
+						OutputSize,
+						InputSize,
+						OutputSize,
+						InputSize);
+				}
+
+				// Do Weighted Sum of Top-2 Experts
+
+				OperatorGatherTopTwoFromSubLayers(
+					OutputBuffer,
+					SparseMixtureOfExpertsInstance->BatchSubLayerIndex0.GetData(),
+					SparseMixtureOfExpertsInstance->BatchSubLayerIndex1.GetData(),
+					SparseMixtureOfExpertsInstance->BatchSubLayerWeight0.GetData(),
+					SparseMixtureOfExpertsInstance->BatchSubLayerWeight1.GetData(),
+					SparseMixtureOfExpertsInstance->BatchSubLayerOutputIndex0.GetData(),
+					SparseMixtureOfExpertsInstance->BatchSubLayerOutputIndex1.GetData(),
+					SparseMixtureOfExpertsInstance->SubLayerOutputBuffers,
+					BatchSize,
+					OutputBufferSize,
+					OutputSize,
+					OutputBufferStride);
+
+				OperatorNanCheck(OutputBuffer, BatchSize, OutputBufferSize, OutputBufferStride);
+			}
+
+			uint32 InputSize = 0;
+			uint32 OutputSize = 0;
+			TSharedPtr<ILayer> GatingLayer;
+			TArray<TSharedPtr<ILayer>, TInlineAllocator<32>> SubLayers;
+		};
+
+		FSparseMixtureOfExpertsLayerInstance::FSparseMixtureOfExpertsLayerInstance(const FSparseMixtureOfExpertsLayer& InSparseMixtureOfExpertsLayer)
+			: SparseMixtureOfExpertsLayer(InSparseMixtureOfExpertsLayer)
+		{
+			GatingInstance = SparseMixtureOfExpertsLayer.GatingLayer->MakeInstance();
+
+			const uint32 SubLayerNum = SparseMixtureOfExpertsLayer.SubLayers.Num();
+
+			SubLayerBatchIndices.SetNum(SubLayerNum);
+			SubLayerInputBuffers.SetNum(SubLayerNum);
+			SubLayerOutputBuffers.SetNum(SubLayerNum);
+
+			SubLayerInstances.Init(nullptr, SubLayerNum);
+			for (uint32 LayerIdx = 0; LayerIdx < SubLayerNum; LayerIdx++)
+			{
+				SubLayerInstances[LayerIdx] = SparseMixtureOfExpertsLayer.SubLayers[LayerIdx]->MakeInstance();
+			}
+		}
+
+		void FSparseMixtureOfExpertsLayerInstance::SetMaxBatchSize(const uint32 MaxBatchSize)
+		{
+			NNE_RUNTIME_BASIC_TRACE_SCOPE(NNE::RuntimeBasic::Private::FSparseMixtureOfExpertsLayerInstance::SetMaxBatchSize);
+
+			const uint32 SubLayerNum = SparseMixtureOfExpertsLayer.SubLayers.Num();
+
+			BatchSubLayerIndex0.SetNumUninitialized(MaxBatchSize, EAllowShrinking::No);
+			BatchSubLayerIndex1.SetNumUninitialized(MaxBatchSize, EAllowShrinking::No);
+			BatchSubLayerWeight0.SetNumUninitialized(MaxBatchSize, EAllowShrinking::No);
+			BatchSubLayerWeight1.SetNumUninitialized(MaxBatchSize, EAllowShrinking::No);
+			BatchSubLayerOutputIndex0.SetNumUninitialized(MaxBatchSize, EAllowShrinking::No);
+			BatchSubLayerOutputIndex1.SetNumUninitialized(MaxBatchSize, EAllowShrinking::No);
+			GatingOutputBuffer.SetNumUninitialized(MaxBatchSize * SubLayerNum, EAllowShrinking::No);
+
+			// Propagate call to sub-layer instances
+
+			if (GatingInstance)
+			{
+				GatingInstance->SetMaxBatchSize(MaxBatchSize);
+			}
+			
+			for (uint32 SubLayerIdx = 0; SubLayerIdx < SubLayerNum; SubLayerIdx++)
+			{
+				if (SubLayerInstances[SubLayerIdx]) { SubLayerInstances[SubLayerIdx]->SetMaxBatchSize(MaxBatchSize); }
+
+				SubLayerBatchIndices[SubLayerIdx].Empty(MaxBatchSize);
+				SubLayerInputBuffers[SubLayerIdx].SetNumUninitialized(MaxBatchSize * SparseMixtureOfExpertsLayer.InputSize, EAllowShrinking::No);
+				SubLayerOutputBuffers[SubLayerIdx].SetNumUninitialized(MaxBatchSize * SparseMixtureOfExpertsLayer.OutputSize, EAllowShrinking::No);
+			}
+		}
+
 		//--------------------------------------------------------------------------
 		// Layer Serialization
 		//--------------------------------------------------------------------------
@@ -3458,6 +3769,7 @@ namespace UE::NNE::RuntimeBasic
 					case ELayerType::AggregateOrExclusive: OutLayer = MakeShared<FAggregateOrExclusiveLayer>(); break;
 					case ELayerType::AggregateOrInclusive: OutLayer = MakeShared<FAggregateOrInclusiveLayer>(); break;
 					case ELayerType::Clamp: OutLayer = MakeShared<FClampLayer>(); break;
+					case ELayerType::SparseMixtureOfExperts: OutLayer = MakeShared<FSparseMixtureOfExpertsLayer>(); break;
 					default: checkf(false, TEXT("Unknown Layer Id %i"), LayerTypeId);
 					}
 				}
@@ -3492,6 +3804,8 @@ namespace UE::NNE::RuntimeBasic
 		, InputTensorDesc(FTensorDesc::Make(TEXT("Input"), FSymbolicTensorShape::Make({ -1, -1 }), ENNETensorDataType::Float))
 		, OutputTensorDesc(FTensorDesc::Make(TEXT("Output"), FSymbolicTensorShape::Make({ -1, -1 }), ENNETensorDataType::Float))
 		, Instance(Model->Layer->MakeInstance())
+		, InputSize(Model->Layer->GetInputSize())
+		, OutputSize(Model->Layer->GetOutputSize())
 	{}
 
 	FModelInstanceCPU::ESetInputTensorShapesStatus FModelInstanceCPU::SetInputTensorShapes(TConstArrayView<FTensorShape> InInputShapes)
@@ -3510,26 +3824,25 @@ namespace UE::NNE::RuntimeBasic
 			return ESetInputTensorShapesStatus::Fail;
 		}
 
-		const uint32 InputBatchSize = InputShape.GetData()[0];
 		const uint32 InputInputSize = InputShape.GetData()[1];
-		const uint32 ModelInputSize = Model->Layer->GetInputSize();
-		const uint32 ModelOutputSize = Model->Layer->GetOutputSize();
 
-		if (!ensureMsgf(InputInputSize == ModelInputSize, TEXT("Input tensor shape does not match model input size. Got %i, expected %i."), InputInputSize, ModelInputSize))
+		if (!ensureMsgf(InputInputSize == InputSize, TEXT("Input tensor shape does not match model input size. Got %i, expected %i."), InputInputSize, InputSize))
 		{
 			return ESetInputTensorShapesStatus::Fail;
 		}
 
-		BatchSize = InputBatchSize;
-		InputSize = ModelInputSize;
-		OutputSize = ModelOutputSize;
-
-		InputTensorShape = FTensorShape::Make({ BatchSize, InputSize });
-		OutputTensorShape = FTensorShape::Make({ BatchSize, OutputSize });
+		const uint32 InputBatchSize = InputShape.GetData()[0];
 
-		if (Instance)
+		if (InputBatchSize != BatchSize)
 		{
-			Instance->SetMaxBatchSize(BatchSize);
+			BatchSize = InputBatchSize;
+			InputTensorShape = FTensorShape::Make({ BatchSize, InputSize });
+			OutputTensorShape = FTensorShape::Make({ BatchSize, OutputSize });
+
+			if (Instance)
+			{
+				Instance->SetMaxBatchSize(BatchSize);
+			}
 		}
 
 		return ESetInputTensorShapesStatus::Ok;
@@ -3683,6 +3996,28 @@ namespace UE::NNE::RuntimeBasic
 		return StaticCastSharedPtr<Private::ILayer>(LinearLayer);
 	}
 
+	FModelBuilderElement FModelBuilder::MakeCompressedLinear(
+		const uint32 InputSize,
+		const uint32 OutputSize,
+		const TConstArrayView<uint16> Weights,
+		const TConstArrayView<float> WeightOffsets,
+		const TConstArrayView<float> WeightScales,
+		const TConstArrayView<float> Biases)
+	{
+		check(Biases.Num() == OutputSize);
+		check(Weights.Num() == InputSize * OutputSize);
+
+		const TSharedPtr<Private::FCompressedLinearLayer> LinearLayer = MakeShared<Private::FCompressedLinearLayer>();
+		LinearLayer->InputSize = InputSize;
+		LinearLayer->OutputSize = OutputSize;
+		LinearLayer->WeightOffsets = WeightOffsets;
+		LinearLayer->WeightScales = WeightScales;
+		LinearLayer->Biases = Biases;
+		LinearLayer->Weights = Weights;
+
+		return StaticCastSharedPtr<Private::ILayer>(LinearLayer);
+	}
+
 	FModelBuilderElement FModelBuilder::MakeLinearWithRandomKaimingWeights(
 		const uint32 InputSize,
 		const uint32 OutputSize,
@@ -3695,6 +4030,25 @@ namespace UE::NNE::RuntimeBasic
 			MakeWeightsZero(OutputSize));
 	}
 
+	FModelBuilderElement FModelBuilder::MakeCompressedLinearWithRandomKaimingWeights(
+		const uint32 InputSize,
+		const uint32 OutputSize,
+		const float WeightScale)
+	{
+		TArrayView<uint16> Weights;
+		TArrayView<float> WeightOffsets;
+		TArrayView<float> WeightScales;
+		MakeCompressedWeightsRandomKaiming(Weights, WeightOffsets, WeightScales, InputSize, OutputSize, WeightScale);
+
+		return MakeCompressedLinear(
+			InputSize,
+			OutputSize,
+			Weights,
+			WeightOffsets,
+			WeightScales,
+			MakeWeightsZero(OutputSize));
+	}
+
 	FModelBuilderElement FModelBuilder::MakeMultiLinear(
 		const uint32 InputSize, 
 		const uint32 OutputSize, 
@@ -3860,6 +4214,35 @@ namespace UE::NNE::RuntimeBasic
 		return MakeSequence(Layers);
 	}
 
+	FModelBuilderElement FModelBuilder::MakeCompressedMLPWithRandomKaimingWeights(
+		const uint32 InputSize,
+		const uint32 OutputSize,
+		const uint32 HiddenSize,
+		const uint32 LayerNum,
+		const EActivationFunction ActivationFunction,
+		const bool bActivationOnFinalLayer)
+	{
+		check(LayerNum >= 2);
+
+		TArray<FModelBuilderElement, TInlineAllocator<32>> Layers;
+		Layers.Reserve(2 * LayerNum - (bActivationOnFinalLayer ? 0 : 1));
+
+		for (uint32 LayerIdx = 0; LayerIdx < LayerNum; LayerIdx++)
+		{
+			const uint32 LayerInputSize = LayerIdx == 0 ? InputSize : HiddenSize;
+			const uint32 LayerOuputSize = LayerIdx == LayerNum - 1 ? OutputSize : HiddenSize;
+
+			Layers.Emplace(MakeCompressedLinearWithRandomKaimingWeights(LayerInputSize, LayerOuputSize));
+
+			if (bActivationOnFinalLayer || LayerIdx != LayerNum - 1)
+			{
+				Layers.Emplace(MakeActivation(LayerOuputSize, ActivationFunction));
+			}
+		}
+
+		return MakeSequence(Layers);
+	}
+
 	FModelBuilderElement FModelBuilder::MakeMemoryCell(
 		const uint32 InputNum,
 		const uint32 OutputNum,
@@ -4075,10 +4458,41 @@ namespace UE::NNE::RuntimeBasic
 		return StaticCastSharedPtr<Private::ILayer>(OrInclusiveLayer);
 	}
 
+	FModelBuilderElement FModelBuilder::MakeSparseMixtureOfExperts(
+		const uint32 InputNum,
+		const uint32 OutputNum,
+		const FModelBuilderElement& GatingLayer,
+		const TConstArrayView<FModelBuilderElement> SubLayers)
+	{
+		check(GatingLayer.GetInputSize() == InputNum);
+		check(GatingLayer.GetOutputSize() == SubLayers.Num());
+
+		const int32 SubLayerNum = SubLayers.Num();
+		for (int32 SubLayerIdx = 0; SubLayerIdx < SubLayerNum; SubLayerIdx++)
+		{
+			check(SubLayers[SubLayerIdx].GetInputSize() == InputNum);
+			check(SubLayers[SubLayerIdx].GetOutputSize() == OutputNum);
+		}
+
+		const TSharedPtr<Private::FSparseMixtureOfExpertsLayer> SparseMixtureOfExpertsLayer = MakeShared<Private::FSparseMixtureOfExpertsLayer>();
+		SparseMixtureOfExpertsLayer->InputSize = InputNum;
+		SparseMixtureOfExpertsLayer->OutputSize = OutputNum;
+		SparseMixtureOfExpertsLayer->GatingLayer = GatingLayer.Layer;
+		SparseMixtureOfExpertsLayer->SubLayers.Reserve(SubLayers.Num());
+
+		for (int32 SubLayerIdx = 0; SubLayerIdx < SubLayerNum; SubLayerIdx++)
+		{
+			SparseMixtureOfExpertsLayer->SubLayers.Add(SubLayers[SubLayerIdx].Layer);
+		}
+
+		return StaticCastSharedPtr<Private::ILayer>(SparseMixtureOfExpertsLayer);
+	}
+
 	void FModelBuilder::Reset()
 	{
 		Rng.Reset();
 		WeightsPool.Empty();
+		CompressedWeightsPool.Empty();
 		SizesPool.Empty();
 	}
 
@@ -4164,6 +4578,82 @@ namespace UE::NNE::RuntimeBasic
 		return Values;
 	}
 
+	void FModelBuilder::MakeCompressedWeightsRandomKaiming(
+		TArrayView<uint16>& OutWeightsView,
+		TArrayView<float>& OutWeightOffsetsView,
+		TArrayView<float>& OutWeightScalesView,
+		const uint32 InputSize,
+		const uint32 OutputSize,
+		const float Scale)
+	{
+		// Make Kaiming Weights
+
+		TArray<float> Values;
+		Values.SetNumUninitialized(InputSize * OutputSize);
+
+		const float Std = Scale * FMath::Sqrt(2.0f / InputSize);
+
+		for (uint32 Idx = 0; Idx < InputSize * OutputSize; Idx++)
+		{
+			Values[Idx] = Std * Private::UniformToGaussian(Rng.FRand(), Rng.FRand());
+		}
+
+		// Find Min and Max
+
+		TArray<float> Mins;
+		TArray<float> Maxs;
+		Mins.SetNumUninitialized(InputSize);
+		Maxs.SetNumUninitialized(InputSize);
+
+		for (uint32 RowIdx = 0; RowIdx < InputSize; RowIdx++)
+		{
+			Mins[RowIdx] = +FLT_MAX;
+			Maxs[RowIdx] = -FLT_MAX;
+		}
+
+		for (uint32 RowIdx = 0; RowIdx < InputSize; RowIdx++)
+		{
+			for (uint32 ColIdx = 0; ColIdx < OutputSize; ColIdx++)
+			{
+				Mins[RowIdx] = FMath::Min(Mins[RowIdx], Values[RowIdx * OutputSize + ColIdx]);
+				Maxs[RowIdx] = FMath::Max(Maxs[RowIdx], Values[RowIdx * OutputSize + ColIdx]);
+			}
+		}
+
+		// Find Scale and Offset
+
+		TArray<float>& WeightOffsets = WeightsPool.AddDefaulted_GetRef();
+		WeightOffsets.SetNumUninitialized(InputSize);
+
+		TArray<float>& WeightScales = WeightsPool.AddDefaulted_GetRef();
+		WeightScales.SetNumUninitialized(InputSize);
+
+		for (uint32 RowIdx = 0; RowIdx < InputSize; RowIdx++)
+		{
+			WeightOffsets[RowIdx] = Mins[RowIdx];
+			WeightScales[RowIdx] = FMath::Max(Maxs[RowIdx] - Mins[RowIdx], UE_SMALL_NUMBER) / 65535.0;
+		}
+
+		// Compress
+
+		TArray<uint16>& Weights = CompressedWeightsPool.AddDefaulted_GetRef();
+		Weights.SetNumUninitialized(InputSize * OutputSize);
+
+		for (uint32 RowIdx = 0; RowIdx < InputSize; RowIdx++)
+		{
+			for (uint32 ColIdx = 0; ColIdx < OutputSize; ColIdx++)
+			{
+				Weights[RowIdx * OutputSize + ColIdx] = (uint16)FMath::RoundToFloat(65535.0 * 
+					FMath::Clamp((Values[RowIdx * OutputSize + ColIdx] - Mins[RowIdx]) / (Maxs[RowIdx] - Mins[RowIdx]), 0.0f, 1.0f));
+			}
+		}
+
+		OutWeightsView = Weights;
+		OutWeightOffsetsView = WeightOffsets;
+		OutWeightScalesView = WeightScales;
+	}
+
+
 	TArrayView<uint32> FModelBuilder::MakeSizesZero(const uint32 Size)
 	{
 		TArray<uint32>& Values = SizesPool.AddDefaulted_GetRef();
diff --git a/Engine/Plugins/Experimental/NNERuntimeBasicCpu/Source/NNERuntimeBasicCpu/Public/NNERuntimeBasicCpuBuilder.h b/Engine/Plugins/Experimental/NNERuntimeBasicCpu/Source/NNERuntimeBasicCpu/Public/NNERuntimeBasicCpuBuilder.h
--- a/Engine/Plugins/Experimental/NNERuntimeBasicCpu/Source/NNERuntimeBasicCpu/Public/NNERuntimeBasicCpuBuilder.h
+++ b/Engine/Plugins/Experimental/NNERuntimeBasicCpu/Source/NNERuntimeBasicCpu/Public/NNERuntimeBasicCpuBuilder.h
@@ -71,6 +71,24 @@ namespace UE::NNE::RuntimeBasic
 			const TConstArrayView<float> Weights,
 			const TConstArrayView<float> Biases);
 
+		/**
+		 * Makes a new compressed linear layer.
+		 *
+		 * @param InputSize		Input Vector Size (Number of Rows)
+		 * @param OutputSize	Output Vector Size (Number of Columns)
+		 * @param Weights		Compressed Linear layer weights.
+		 * @param WeightOffsets	Compressed Linear layer weight offsets.
+		 * @param WeightScales	Compressed Linear layer weight biases.
+		 * @param Biases		Compressed Linear layer biases.
+		 */
+		FModelBuilderElement MakeCompressedLinear(
+			const uint32 InputSize,
+			const uint32 OutputSize,
+			const TConstArrayView<uint16> Weights,
+			const TConstArrayView<float> WeightOffsets,
+			const TConstArrayView<float> WeightScales,
+			const TConstArrayView<float> Biases);
+
 		/**
 		 * Makes a new linear layer with randomly initialized weights using the Kaiming method.
 		 *
@@ -83,6 +101,18 @@ namespace UE::NNE::RuntimeBasic
 			const uint32 OutputSize,
 			const float WeightScale = 1.0f);
 
+		/**
+		 * Makes a new compressed linear layer with randomly initialized weights using the Kaiming method.
+		 *
+		 * @param InputSize		Input Vector Size (Number of Rows)
+		 * @param OutputSize	Output Vector Size (Number of Columns)
+		 * @param WeightScale	Scaling factor for the weight creation.
+		 */
+		FModelBuilderElement MakeCompressedLinearWithRandomKaimingWeights(
+			const uint32 InputSize,
+			const uint32 OutputSize,
+			const float WeightScale = 1.0f);
+
 		/**
 		 * Makes a new multi linear layer.
 		 *
@@ -169,6 +199,24 @@ namespace UE::NNE::RuntimeBasic
 			const EActivationFunction ActivationFunction,
 			const bool bActivationOnFinalLayer = false);
 
+		/**
+		 * Makes a Compressed Multi-Layer Perceptron network with randomly initialized weights using the Kaiming method.
+		 *
+		 * @param InputSize					Input Vector Size
+		 * @param OutputSize				Output Vector Size
+		 * @param HiddenSize				Number of hidden units to use on internal layers.
+		 * @param LayerNum					Number of layers. Includes input and output layers.
+		 * @param ActivationFunction		Activation function to use.
+		 * @param bActivationOnFinalLayer	If the activation function should be used on the final output layer.
+		 */
+		FModelBuilderElement MakeCompressedMLPWithRandomKaimingWeights(
+			const uint32 InputSize,
+			const uint32 OutputSize,
+			const uint32 HiddenSize,
+			const uint32 LayerNum,
+			const EActivationFunction ActivationFunction,
+			const bool bActivationOnFinalLayer = false);
+
 		/**
 		 * Make a new Memory Cell layer.
 		 *
@@ -266,6 +314,20 @@ namespace UE::NNE::RuntimeBasic
 			const TConstArrayView<FModelBuilderElement> KeyLayers,
 			const TConstArrayView<FModelBuilderElement> ValueLayers);
 
+		/**
+		 * Make a new Top-Two Sparse Mixture of Experts Layer.
+		 *
+		 * @param InputNum					Number of normal inputs to the model.
+		 * @param OutputNum					Number of normal outputs from the model.
+		 * @param GatingLayer				Layer used to choose the experts. Input should be of size InputNum, output should be of size SubLayerNum.
+		 * @param SubLayers					Expert layers. All Layers here should have input of size InputNum, output of size OutputNum.
+		 */
+		FModelBuilderElement MakeSparseMixtureOfExperts(
+			const uint32 InputNum,
+			const uint32 OutputNum,
+			const FModelBuilderElement& GatingLayer,
+			const TConstArrayView<FModelBuilderElement> SubLayers);
+
 	public:
 
 		/** Creates a array of weights from a copy of the given array view */
@@ -280,6 +342,15 @@ namespace UE::NNE::RuntimeBasic
 		/** Creates a array of weights randomly initialized using the Kaiming method */
 		TArrayView<float> MakeWeightsRandomKaiming(const uint32 InputSize, const uint32 OutputSize, const float Scale = 1.0f);
 
+		/** Creates a array of weights randomly initialized using the Kaiming method */
+		void MakeCompressedWeightsRandomKaiming(
+			TArrayView<uint16>& OutWeightsView,
+			TArrayView<float>& OutWeightOffsetsView,
+			TArrayView<float>& OutWeightScalesView,
+			const uint32 InputSize,
+			const uint32 OutputSize,
+			const float Scale = 1.0f);
+
 		/** Creates an array of sizes, initialized to zero. */
 		TArrayView<uint32> MakeSizesZero(const uint32 Size);
 
@@ -349,6 +420,9 @@ namespace UE::NNE::RuntimeBasic
 		/** Pool of all weights data used by the `MakeWeights` functions. */
 		TArray<TArray<float>> WeightsPool;
 
+		/** Pool of all compressed weights data used by the `MakeCompressedWeights` functions. */
+		TArray<TArray<uint16>> CompressedWeightsPool;
+
 		/** Pool of all sizes data used by the `MakeSizes` functions. */
 		TArray<TArray<uint32>> SizesPool;
 	};
