NNERuntimeBasicCpu: Added GELU Activation

diff --git a/Engine/Plugins/Experimental/NNERuntimeBasicCpu/Source/NNERuntimeBasicCpu/Private/NNERuntimeBasicCpu.ispc b/Engine/Plugins/Experimental/NNERuntimeBasicCpu/Source/NNERuntimeBasicCpu/Private/NNERuntimeBasicCpu.ispc
--- a/Engine/Plugins/Experimental/NNERuntimeBasicCpu/Source/NNERuntimeBasicCpu/Private/NNERuntimeBasicCpu.ispc
+++ b/Engine/Plugins/Experimental/NNERuntimeBasicCpu/Source/NNERuntimeBasicCpu/Private/NNERuntimeBasicCpu.ispc
@@ -11,6 +11,11 @@ inline float NNERuntimeBasicCPUSigmoid(const float X)
 	return 1.0f / (1.0f + exp(-X));
 }
 
+inline float NNERuntimeBasicCPUGELU(const float X)
+{
+	return X * NNERuntimeBasicCPUSigmoid(1.702f * X);
+}
+
 export void NNERuntimeBasicCPUOperatorNormalize(
 	uniform float Output[],
 	uniform const float Input[],
@@ -210,6 +215,20 @@ export void NNERuntimeBasicCPUOperatorELU(
 	}
 }
 
+export void NNERuntimeBasicCPUOperatorGELU(
+	uniform float Output[],
+	uniform const float Input[],
+	uniform const int BatchNum,
+	uniform const int Num,
+	uniform const int OutputStride,
+	uniform const int InputStride)
+{
+	foreach(BatchIdx = 0 ... BatchNum, Idx = 0 ... Num)
+	{
+		Output[BatchIdx * OutputStride + Idx] = NNERuntimeBasicCPUGELU(Input[BatchIdx * InputStride + Idx]);
+	}
+}
+
 export void NNERuntimeBasicCPUOperatorTanH(
 	uniform float Output[],
 	uniform const float Input[],
diff --git a/Engine/Plugins/Experimental/NNERuntimeBasicCpu/Source/NNERuntimeBasicCpu/Private/NNERuntimeBasicCpuModel.cpp b/Engine/Plugins/Experimental/NNERuntimeBasicCpu/Source/NNERuntimeBasicCpu/Private/NNERuntimeBasicCpuModel.cpp
--- a/Engine/Plugins/Experimental/NNERuntimeBasicCpu/Source/NNERuntimeBasicCpu/Private/NNERuntimeBasicCpuModel.cpp
+++ b/Engine/Plugins/Experimental/NNERuntimeBasicCpu/Source/NNERuntimeBasicCpu/Private/NNERuntimeBasicCpuModel.cpp
@@ -537,6 +537,40 @@ namespace UE::NNE::RuntimeBasic
 #endif
 		}
 
+		static inline float GELU(const float X)
+		{
+			return X * Sigmoid(1.702f * X);
+		}
+
+		static inline void OperatorGELU(
+			float* RESTRICT Output,
+			const float* RESTRICT Input,
+			const uint32 BatchSize,
+			const uint32 InputOutputSize,
+			const uint32 OutputStride,
+			const uint32 InputStride)
+		{
+			NNE_RUNTIME_BASIC_TRACE_SCOPE(NNE::RuntimeBasic::Private::OperatorGELU);
+
+#if NNE_RUNTIME_BASIC_ENABLE_ISPC
+			ispc::NNERuntimeBasicCPUOperatorGELU(
+				Output,
+				Input,
+				BatchSize,
+				InputOutputSize,
+				OutputStride,
+				InputStride);
+#else
+			for (uint32 BatchIdx = 0; BatchIdx < BatchSize; BatchIdx++)
+			{
+				for (uint32 Idx = 0; Idx < InputOutputSize; Idx++)
+				{
+					Output[BatchIdx * OutputStride + Idx] = GELU(Input[BatchIdx * InputStride + Idx]);
+				}
+			}
+#endif
+		}
+
 		static inline void OperatorTanH(
 			float* RESTRICT Output,
 			const float* RESTRICT Input,
@@ -1291,6 +1325,7 @@ namespace UE::NNE::RuntimeBasic
 
 			Clamp = 18,
 			SparseMixtureOfExperts = 19,
+			GELU = 20,
 		};
 
 		//--------------------------------------------------------------------------
@@ -1999,6 +2034,58 @@ namespace UE::NNE::RuntimeBasic
 
 		//--------------------------------------------------------------------------
 
+		struct FGELULayer : public ILayer
+		{
+			virtual ELayerType GetLayerType() const override final { return ELayerType::GELU; }
+			virtual uint32 GetInputSize() const override final { return InputOutputSize; }
+			virtual uint32 GetOutputSize() const override final { return InputOutputSize; }
+
+			virtual void SerializationSize(uint64& InOutOffset) const override final
+			{
+				Serialization::Size(InOutOffset, InputOutputSize);
+			}
+
+			virtual void SerializationLoad(uint64& InOutOffset, TConstArrayView<uint8> Data) override final
+			{
+				Serialization::Load(InOutOffset, InputOutputSize, Data);
+			}
+
+			virtual void SerializationSave(uint64& InOutOffset, TArrayView<uint8> Data) const override final
+			{
+				Serialization::Save(InOutOffset, InputOutputSize, Data);
+			}
+
+			virtual void Evaluate(
+				ILayerInstance* Instance,
+				float* OutputBuffer,
+				const float* InputBuffer,
+				const uint32 BatchSize,
+				const uint32 OutputBufferSize,
+				const uint32 InputBufferSize,
+				const uint32 OutputBufferStride,
+				const uint32 InputBufferStride) override final
+			{
+				NNE_RUNTIME_BASIC_TRACE_SCOPE(NNE::RuntimeBasic::Private::FGELULayer::Evaluate);
+				check(OutputBufferSize == GetOutputSize() && InputBufferSize == GetInputSize());
+				check(Instance == nullptr);
+				OperatorNanCheck(InputBuffer, BatchSize, InputBufferSize, InputBufferStride);
+
+				OperatorGELU(
+					OutputBuffer,
+					InputBuffer,
+					BatchSize,
+					InputOutputSize,
+					OutputBufferStride,
+					InputBufferStride);
+
+				OperatorNanCheck(OutputBuffer, BatchSize, OutputBufferSize, OutputBufferStride);
+			}
+
+			uint32 InputOutputSize = 0;
+		};
+
+		//--------------------------------------------------------------------------
+
 		struct FTanHLayer : public ILayer
 		{
 			virtual ELayerType GetLayerType() const override final { return ELayerType::TanH; }
@@ -3770,6 +3857,7 @@ namespace UE::NNE::RuntimeBasic
 					case ELayerType::AggregateOrInclusive: OutLayer = MakeShared<FAggregateOrInclusiveLayer>(); break;
 					case ELayerType::Clamp: OutLayer = MakeShared<FClampLayer>(); break;
 					case ELayerType::SparseMixtureOfExperts: OutLayer = MakeShared<FSparseMixtureOfExpertsLayer>(); break;
+					case ELayerType::GELU: OutLayer = MakeShared<FGELULayer>(); break;
 					default: checkf(false, TEXT("Unknown Layer Id %i"), LayerTypeId);
 					}
 				}
@@ -3801,9 +3889,12 @@ namespace UE::NNE::RuntimeBasic
 
 	FModelInstanceCPU::FModelInstanceCPU(const TSharedPtr<FModelCPU>& InModel)
 		: Model(InModel)
-		, InputTensorDesc(FTensorDesc::Make(TEXT("Input"), FSymbolicTensorShape::Make({ -1, -1 }), ENNETensorDataType::Float))
-		, OutputTensorDesc(FTensorDesc::Make(TEXT("Output"), FSymbolicTensorShape::Make({ -1, -1 }), ENNETensorDataType::Float))
+		, InputTensorDesc(FTensorDesc::Make(TEXT("Input"), FSymbolicTensorShape::Make({ -1, (int32)Model->Layer->GetInputSize() }), ENNETensorDataType::Float))
+		, OutputTensorDesc(FTensorDesc::Make(TEXT("Output"), FSymbolicTensorShape::Make({ -1, (int32)Model->Layer->GetOutputSize() }), ENNETensorDataType::Float))
+		, InputTensorShape(FTensorShape::Make({ 0, Model->Layer->GetInputSize() }))
+		, OutputTensorShape(FTensorShape::Make({ 0, Model->Layer->GetOutputSize() }))
 		, Instance(Model->Layer->MakeInstance())
+		, BatchSize(0)
 		, InputSize(Model->Layer->GetInputSize())
 		, OutputSize(Model->Layer->GetOutputSize())
 	{}
@@ -4115,6 +4206,13 @@ namespace UE::NNE::RuntimeBasic
 		return StaticCastSharedPtr<Private::ILayer>(Layer);
 	}
 
+	FModelBuilderElement FModelBuilder::MakeGELU(const uint32 InputOutputSize)
+	{
+		const TSharedPtr<Private::FGELULayer> Layer = MakeShared<Private::FGELULayer>();
+		Layer->InputOutputSize = InputOutputSize;
+		return StaticCastSharedPtr<Private::ILayer>(Layer);
+	}
+
 	FModelBuilderElement FModelBuilder::MakeTanH(const uint32 InputOutputSize)
 	{
 		const TSharedPtr<Private::FTanHLayer> Layer = MakeShared<Private::FTanHLayer>();
@@ -4148,6 +4246,7 @@ namespace UE::NNE::RuntimeBasic
 		case EActivationFunction::ReLU: return MakeReLU(InputOutputSize);
 		case EActivationFunction::ELU: return MakeELU(InputOutputSize);
 		case EActivationFunction::TanH: return MakeTanH(InputOutputSize);
+		case EActivationFunction::GELU: return MakeGELU(InputOutputSize);
 		default:
 			checkf(false, TEXT("Unknown Activation Function"));
 			return MakeReLU(InputOutputSize);
diff --git a/Engine/Plugins/Experimental/NNERuntimeBasicCpu/Source/NNERuntimeBasicCpu/Public/NNERuntimeBasicCpuBuilder.h b/Engine/Plugins/Experimental/NNERuntimeBasicCpu/Source/NNERuntimeBasicCpu/Public/NNERuntimeBasicCpuBuilder.h
--- a/Engine/Plugins/Experimental/NNERuntimeBasicCpu/Source/NNERuntimeBasicCpu/Public/NNERuntimeBasicCpuBuilder.h
+++ b/Engine/Plugins/Experimental/NNERuntimeBasicCpu/Source/NNERuntimeBasicCpu/Public/NNERuntimeBasicCpuBuilder.h
@@ -52,6 +52,7 @@ namespace UE::NNE::RuntimeBasic
 			ReLU = 0,
 			ELU = 1,
 			TanH = 2,
+			GELU = 3,
 		};
 
 		/** Construct a new Model Builder with the given random seed. */
@@ -156,6 +157,9 @@ namespace UE::NNE::RuntimeBasic
 		/** Makes a ELU Activation Layer */
 		FModelBuilderElement MakeELU(const uint32 InputOutputSize);
 
+		/** Makes a GELU Activation Layer */
+		FModelBuilderElement MakeGELU(const uint32 InputOutputSize);
+
 		/** Makes a TanH Activation Layer */
 		FModelBuilderElement MakeTanH(const uint32 InputOutputSize);
 
