NNERuntimeBasicCpu: Added aliasing and stride checks. Fixed bug with incorrect max activation stride computation for sequence layer.

diff --git a/Engine/Plugins/Experimental/NNERuntimeBasicCpu/Source/NNERuntimeBasicCpu/Private/NNERuntimeBasicCpuModel.cpp b/Engine/Plugins/Experimental/NNERuntimeBasicCpu/Source/NNERuntimeBasicCpu/Private/NNERuntimeBasicCpuModel.cpp
--- a/Engine/Plugins/Experimental/NNERuntimeBasicCpu/Source/NNERuntimeBasicCpu/Private/NNERuntimeBasicCpuModel.cpp
+++ b/Engine/Plugins/Experimental/NNERuntimeBasicCpu/Source/NNERuntimeBasicCpu/Private/NNERuntimeBasicCpuModel.cpp
@@ -203,6 +203,8 @@ namespace UE::NNE::RuntimeBasic
 		{
 			NNE_RUNTIME_BASIC_TRACE_SCOPE(NNE::RuntimeBasic::Private::OperatorCopy);
 
+			check(Output != Input);
+
 #if NNE_RUNTIME_BASIC_ENABLE_ISPC
 			ispc::NNERuntimeBasicCPUOperatorCopy(
 				Output,
@@ -234,6 +236,8 @@ namespace UE::NNE::RuntimeBasic
 		{
 			NNE_RUNTIME_BASIC_TRACE_SCOPE(NNE::RuntimeBasic::Private::OperatorNormalize);
 
+			check(Output != Input);
+
 #if NNE_RUNTIME_BASIC_ENABLE_ISPC
 			ispc::NNERuntimeBasicCPUOperatorNormalize(
 				Output,
@@ -267,6 +271,8 @@ namespace UE::NNE::RuntimeBasic
 		{
 			NNE_RUNTIME_BASIC_TRACE_SCOPE(NNE::RuntimeBasic::Private::OperatorDenormalize);
 
+			check(Output != Input);
+
 #if NNE_RUNTIME_BASIC_ENABLE_ISPC
 			ispc::NNERuntimeBasicCPUOperatorDenormalize(
 				Output,
@@ -300,6 +306,8 @@ namespace UE::NNE::RuntimeBasic
 		{
 			NNE_RUNTIME_BASIC_TRACE_SCOPE(NNE::RuntimeBasic::Private::OperatorClamp);
 
+			check(Output != Input);
+
 #if NNE_RUNTIME_BASIC_ENABLE_ISPC
 			ispc::NNERuntimeBasicCPUOperatorClamp(
 				Output,
@@ -334,6 +342,8 @@ namespace UE::NNE::RuntimeBasic
 		{
 			NNE_RUNTIME_BASIC_TRACE_SCOPE(NNE::RuntimeBasic::Private::OperatorLinear);
 
+			check(Output != Input);
+
 #if NNE_RUNTIME_BASIC_ENABLE_ISPC
 			ispc::NNERuntimeBasicCPUOperatorLinear(
 				Output,
@@ -387,6 +397,8 @@ namespace UE::NNE::RuntimeBasic
 		{
 			NNE_RUNTIME_BASIC_TRACE_SCOPE(NNE::RuntimeBasic::Private::OperatorCompressedLinear);
 
+			check(Output != Input);
+
 #if NNE_RUNTIME_BASIC_ENABLE_ISPC
 			ispc::NNERuntimeBasicCPUOperatorCompressedLinear(
 				Output,
@@ -444,6 +456,8 @@ namespace UE::NNE::RuntimeBasic
 		{
 			NNE_RUNTIME_BASIC_TRACE_SCOPE(NNE::RuntimeBasic::Private::OperatorMultiLinear);
 
+			check(Output != Input);
+
 			// For this function ispc generates slightly less efficient code than the naive C++ implementation so we
 			// don't bother calling out to the ispc version even if it is available
 
@@ -488,6 +502,8 @@ namespace UE::NNE::RuntimeBasic
 		{
 			NNE_RUNTIME_BASIC_TRACE_SCOPE(NNE::RuntimeBasic::Private::OperatorReLU);
 
+			check(Output != Input);
+
 #if NNE_RUNTIME_BASIC_ENABLE_ISPC
 			ispc::NNERuntimeBasicCPUOperatorReLU(
 				Output,
@@ -517,6 +533,8 @@ namespace UE::NNE::RuntimeBasic
 		{
 			NNE_RUNTIME_BASIC_TRACE_SCOPE(NNE::RuntimeBasic::Private::OperatorELU);
 
+			check(Output != Input);
+
 #if NNE_RUNTIME_BASIC_ENABLE_ISPC
 			ispc::NNERuntimeBasicCPUOperatorELU(
 				Output,
@@ -552,6 +570,8 @@ namespace UE::NNE::RuntimeBasic
 		{
 			NNE_RUNTIME_BASIC_TRACE_SCOPE(NNE::RuntimeBasic::Private::OperatorGELU);
 
+			check(Output != Input);
+
 #if NNE_RUNTIME_BASIC_ENABLE_ISPC
 			ispc::NNERuntimeBasicCPUOperatorGELU(
 				Output,
@@ -581,6 +601,8 @@ namespace UE::NNE::RuntimeBasic
 		{
 			NNE_RUNTIME_BASIC_TRACE_SCOPE(NNE::RuntimeBasic::Private::OperatorTanH);
 
+			check(Output != Input);
+
 #if NNE_RUNTIME_BASIC_ENABLE_ISPC
 			ispc::NNERuntimeBasicCPUOperatorTanH(
 				Output,
@@ -611,6 +633,8 @@ namespace UE::NNE::RuntimeBasic
 		{
 			NNE_RUNTIME_BASIC_TRACE_SCOPE(NNE::RuntimeBasic::Private::OperatorPReLU);
 
+			check(Output != Input);
+
 #if NNE_RUNTIME_BASIC_ENABLE_ISPC
 			ispc::NNERuntimeBasicCPUOperatorPReLU(
 				Output,
@@ -646,6 +670,10 @@ namespace UE::NNE::RuntimeBasic
 		{
 			NNE_RUNTIME_BASIC_TRACE_SCOPE(NNE::RuntimeBasic::Private::OperatorMemoryCellUpdateMemory);
 
+			check(Output != RememberGate);
+			check(Output != Memory);
+			check(Output != Update);
+
 #if NNE_RUNTIME_BASIC_ENABLE_ISPC
 			ispc::NNERuntimeBasicCPUOperatorMemoryCellUpdateMemory(
 				Output,
@@ -687,6 +715,10 @@ namespace UE::NNE::RuntimeBasic
 		{
 			NNE_RUNTIME_BASIC_TRACE_SCOPE(NNE::RuntimeBasic::Private::OperatorMemoryCellUpdateOutput);
 
+			check(Output != PassthroughGate);
+			check(Output != MemoryUpdate);
+			check(Output != InputUpdate);
+
 #if NNE_RUNTIME_BASIC_ENABLE_ISPC
 			ispc::NNERuntimeBasicCPUOperatorMemoryCellUpdateOutput(
 				Output,
@@ -725,6 +757,8 @@ namespace UE::NNE::RuntimeBasic
 		{
 			NNE_RUNTIME_BASIC_TRACE_SCOPE(NNE::RuntimeBasic::Private::OperatorAggregateGatherElements);
 
+			check(OutputBuffer != InputBuffer);
+
 #if NNE_RUNTIME_BASIC_ENABLE_ISPC
 			ispc::NNERuntimeBasicCPUOperatorAggregateGatherElements(
 				OutputBuffer,
@@ -836,6 +870,12 @@ namespace UE::NNE::RuntimeBasic
 		{
 			NNE_RUNTIME_BASIC_TRACE_SCOPE(NNE::RuntimeBasic::Private::OperatorAggregateGatherFromSubLayers);
 
+			check(QueryBuffer != KeyBuffer);
+			check(QueryBuffer != ValueBuffer);
+			check(KeyBuffer != ValueBuffer);
+			check(ElementAccum != ElementNums);
+			check(ElementAccum != ElementOffsets);
+
 			const uint32 SubLayerNum = SubLayerBatchIndices.Num();
 
 			for (uint32 BatchIdx = 0; BatchIdx < BatchSize; BatchIdx++)
@@ -908,6 +948,9 @@ namespace UE::NNE::RuntimeBasic
 		{
 			NNE_RUNTIME_BASIC_TRACE_SCOPE(NNE::RuntimeBasic::Private::OperatorAggregateDotProductAttention);
 
+			check(Attention != Queries);
+			check(Attention != Keys);
+
 #if NNE_RUNTIME_BASIC_ENABLE_ISPC
 			ispc::NNERuntimeBasicCPUOperatorAggregateDotProductAttention(
 				Attention,
@@ -960,6 +1003,10 @@ namespace UE::NNE::RuntimeBasic
 		{
 			NNE_RUNTIME_BASIC_TRACE_SCOPE(NNE::RuntimeBasic::Private::OperatorAggregateSoftmaxPlusOneInplace);
 
+			check(AttentionMaxs != AttentionDenoms);
+			check(AttentionMaxs != Attention);
+			check(Attention != AttentionDenoms);
+
 			// Numerically stable soft-max computation using subtraction of the (positive) max value
 			// 
 			// Here the +1 in the denominator allows the attention 
@@ -1027,6 +1074,9 @@ namespace UE::NNE::RuntimeBasic
 		{
 			NNE_RUNTIME_BASIC_TRACE_SCOPE(NNE::RuntimeBasic::Private::OperatorAggregateAttentionSum);
 
+			check(Output != Attention);
+			check(Output != Values);
+
 #if NNE_RUNTIME_BASIC_ENABLE_ISPC
 			ispc::NNERuntimeBasicCPUOperatorAggregateAttentionSum(
 				Output,
@@ -1082,6 +1132,8 @@ namespace UE::NNE::RuntimeBasic
 		{
 			NNE_RUNTIME_BASIC_TRACE_SCOPE(NNE::RuntimeBasic::Private::OperatorGather);
 
+			check(OutputBuffer != InputBuffer);
+
 #if NNE_RUNTIME_BASIC_ENABLE_ISPC
 			ispc::NNERuntimeBasicCPUOperatorGather(
 				OutputBuffer,
@@ -1114,6 +1166,8 @@ namespace UE::NNE::RuntimeBasic
 		{
 			NNE_RUNTIME_BASIC_TRACE_SCOPE(NNE::RuntimeBasic::Private::OperatorScatter);
 
+			check(OutputBuffer != InputBuffer);
+
 #if NNE_RUNTIME_BASIC_ENABLE_ISPC
 			ispc::NNERuntimeBasicCPUOperatorScatter(
 				OutputBuffer,
@@ -1207,6 +1261,10 @@ namespace UE::NNE::RuntimeBasic
 		{
 			NNE_RUNTIME_BASIC_TRACE_SCOPE(NNE::RuntimeBasic::Private::OperatorGatherTopTwoSubLayerBatchIndices);
 
+			check(BatchSubLayerIndex0 != BatchSubLayerIndex1);
+			check(BatchSubLayerWeight0 != BatchSubLayerWeight1);
+			check(BatchSubLayerOutputIndex0 != BatchSubLayerOutputIndex1);
+
 			check(SubLayerGateSize >= 2);
 
 			for (uint32 SubLayerIdx = 0; SubLayerIdx < SubLayerGateSize; SubLayerIdx++)
@@ -1462,6 +1520,7 @@ namespace UE::NNE::RuntimeBasic
 			{
 				NNE_RUNTIME_BASIC_TRACE_SCOPE(NNE::RuntimeBasic::Private::FSequenceLayer::Evaluate);
 				check(OutputBufferSize == GetOutputSize() && InputBufferSize == GetInputSize());
+				check(OutputBufferStride >= GetOutputSize() && InputBufferStride >= GetInputSize());
 				OperatorNanCheck(InputBuffer, BatchSize, InputBufferSize, InputBufferStride);
 
 				const uint32 LayerNum = Layers.Num();
@@ -1559,10 +1618,17 @@ namespace UE::NNE::RuntimeBasic
 
 			// Compute the largest intermediate size used
 
-			ActivationStride = SequenceLayer.Layers[0]->GetInputSize();
-			for (uint32 LayerIdx = 1; LayerIdx < LayerNum; LayerIdx++)
+			if (LayerNum == 0)
 			{
-				ActivationStride = FMath::Max(ActivationStride, SequenceLayer.Layers[LayerIdx]->GetOutputSize());
+				ActivationStride = 0;
+			}
+			else
+			{
+				ActivationStride = SequenceLayer.Layers[0]->GetOutputSize();
+				for (uint32 LayerIdx = 1; LayerIdx < LayerNum - 1; LayerIdx++)
+				{
+					ActivationStride = FMath::Max(ActivationStride, SequenceLayer.Layers[LayerIdx]->GetOutputSize());
+				}
 			}
 		}
 
@@ -1629,6 +1695,7 @@ namespace UE::NNE::RuntimeBasic
 			{
 				NNE_RUNTIME_BASIC_TRACE_SCOPE(NNE::RuntimeBasic::Private::FNormalizeLayer::Evaluate);
 				check(OutputBufferSize == GetOutputSize() && InputBufferSize == GetInputSize());
+				check(OutputBufferStride >= GetOutputSize() && InputBufferStride >= GetInputSize());
 				check(Instance == nullptr);
 				OperatorNanCheck(InputBuffer, BatchSize, InputBufferSize, InputBufferStride);
 
@@ -1691,6 +1758,7 @@ namespace UE::NNE::RuntimeBasic
 			{
 				NNE_RUNTIME_BASIC_TRACE_SCOPE(NNE::RuntimeBasic::Private::FDenormalizeLayer::Evaluate);
 				check(OutputBufferSize == GetOutputSize() && InputBufferSize == GetInputSize());
+				check(OutputBufferStride >= GetOutputSize() && InputBufferStride >= GetInputSize());
 				check(Instance == nullptr);
 				OperatorNanCheck(InputBuffer, BatchSize, InputBufferSize, InputBufferStride);
 
@@ -1756,6 +1824,7 @@ namespace UE::NNE::RuntimeBasic
 			{
 				NNE_RUNTIME_BASIC_TRACE_SCOPE(NNE::RuntimeBasic::Private::FLinearLayer::Evaluate);
 				check(OutputBufferSize == GetOutputSize() && InputBufferSize == GetInputSize());
+				check(OutputBufferStride >= GetOutputSize() && InputBufferStride >= GetInputSize());
 				check(Instance == nullptr);
 				OperatorNanCheck(InputBuffer, BatchSize, InputBufferSize, InputBufferStride);
 
@@ -1829,6 +1898,7 @@ namespace UE::NNE::RuntimeBasic
 			{
 				NNE_RUNTIME_BASIC_TRACE_SCOPE(NNE::RuntimeBasic::Private::FCompressedLinearLayer::Evaluate);
 				check(OutputBufferSize == GetOutputSize() && InputBufferSize == GetInputSize());
+				check(OutputBufferStride >= GetOutputSize() && InputBufferStride >= GetInputSize());
 				check(Instance == nullptr);
 				OperatorNanCheck(InputBuffer, BatchSize, InputBufferSize, InputBufferStride);
 
@@ -1903,6 +1973,7 @@ namespace UE::NNE::RuntimeBasic
 			{
 				NNE_RUNTIME_BASIC_TRACE_SCOPE(NNE::RuntimeBasic::Private::FMultiLinearLayer::Evaluate);
 				check(OutputBufferSize == GetOutputSize() && InputBufferSize == GetInputSize());
+				check(OutputBufferStride >= GetOutputSize() && InputBufferStride >= GetInputSize());
 				check(Instance == nullptr);
 				OperatorNanCheck(InputBuffer, BatchSize, InputBufferSize, InputBufferStride);
 
@@ -1963,6 +2034,7 @@ namespace UE::NNE::RuntimeBasic
 			{
 				NNE_RUNTIME_BASIC_TRACE_SCOPE(NNE::RuntimeBasic::Private::FReLULayer::Evaluate);
 				check(OutputBufferSize == GetOutputSize() && InputBufferSize == GetInputSize());
+				check(OutputBufferStride >= GetOutputSize() && InputBufferStride >= GetInputSize());
 				check(Instance == nullptr);
 				OperatorNanCheck(InputBuffer, BatchSize, InputBufferSize, InputBufferStride);
 
@@ -2015,6 +2087,7 @@ namespace UE::NNE::RuntimeBasic
 			{
 				NNE_RUNTIME_BASIC_TRACE_SCOPE(NNE::RuntimeBasic::Private::FELULayer::Evaluate);
 				check(OutputBufferSize == GetOutputSize() && InputBufferSize == GetInputSize());
+				check(OutputBufferStride >= GetOutputSize() && InputBufferStride >= GetInputSize());
 				check(Instance == nullptr);
 				OperatorNanCheck(InputBuffer, BatchSize, InputBufferSize, InputBufferStride);
 
@@ -2067,6 +2140,7 @@ namespace UE::NNE::RuntimeBasic
 			{
 				NNE_RUNTIME_BASIC_TRACE_SCOPE(NNE::RuntimeBasic::Private::FGELULayer::Evaluate);
 				check(OutputBufferSize == GetOutputSize() && InputBufferSize == GetInputSize());
+				check(OutputBufferStride >= GetOutputSize() && InputBufferStride >= GetInputSize());
 				check(Instance == nullptr);
 				OperatorNanCheck(InputBuffer, BatchSize, InputBufferSize, InputBufferStride);
 
@@ -2119,6 +2193,7 @@ namespace UE::NNE::RuntimeBasic
 			{
 				NNE_RUNTIME_BASIC_TRACE_SCOPE(NNE::RuntimeBasic::Private::FTanHLayer::Evaluate);
 				check(OutputBufferSize == GetOutputSize() && InputBufferSize == GetInputSize());
+				check(OutputBufferStride >= GetOutputSize() && InputBufferStride >= GetInputSize());
 				check(Instance == nullptr);
 				OperatorNanCheck(InputBuffer, BatchSize, InputBufferSize, InputBufferStride);
 
@@ -2174,6 +2249,7 @@ namespace UE::NNE::RuntimeBasic
 			{
 				NNE_RUNTIME_BASIC_TRACE_SCOPE(NNE::RuntimeBasic::Private::FPReLuLayer::Evaluate);
 				check(OutputBufferSize == GetOutputSize() && InputBufferSize == GetInputSize());
+				check(OutputBufferStride >= GetOutputSize() && InputBufferStride >= GetInputSize());
 				check(Instance == nullptr);
 				OperatorNanCheck(InputBuffer, BatchSize, InputBufferSize, InputBufferStride);
 
@@ -2271,6 +2347,7 @@ namespace UE::NNE::RuntimeBasic
 			{
 				NNE_RUNTIME_BASIC_TRACE_SCOPE(NNE::RuntimeBasic::Private::FMemoryCellLayer::Evaluate);
 				check(OutputBufferSize == GetOutputSize() && InputBufferSize == GetInputSize());
+				check(OutputBufferStride >= GetOutputSize() && InputBufferStride >= GetInputSize());
 
 				FMemoryCellInstance* MemoryCellInstance = StaticCast<FMemoryCellInstance*>(Instance);
 				check(MemoryCellInstance);
@@ -2440,6 +2517,7 @@ namespace UE::NNE::RuntimeBasic
 			{
 				NNE_RUNTIME_BASIC_TRACE_SCOPE(NNE::RuntimeBasic::Private::FCopyLayer::Evaluate);
 				check(OutputBufferSize == GetOutputSize() && InputBufferSize == GetInputSize());
+				check(OutputBufferStride >= GetOutputSize() && InputBufferStride >= GetInputSize());
 				check(Instance == nullptr);
 				OperatorNanCheck(InputBuffer, BatchSize, InputBufferSize, InputBufferStride);
 
@@ -2534,6 +2612,7 @@ namespace UE::NNE::RuntimeBasic
 			{
 				NNE_RUNTIME_BASIC_TRACE_SCOPE(NNE::RuntimeBasic::Private::FConcatLayer::Evaluate);
 				check(OutputBufferSize == GetOutputSize() && InputBufferSize == GetInputSize());
+				check(OutputBufferStride >= GetOutputSize() && InputBufferStride >= GetInputSize());
 
 				FConcatLayerInstance* ConcatInstance = StaticCast<FConcatLayerInstance*>(Instance);
 				check(ConcatInstance);
@@ -2544,6 +2623,9 @@ namespace UE::NNE::RuntimeBasic
 
 				for (int32 LayerIdx = 0; LayerIdx < LayerNum; LayerIdx++)
 				{
+					check(InputOffsets[LayerIdx] + InputSizes[LayerIdx] <= InputBufferSize);
+					check(OutputOffsets[LayerIdx] + OutputSizes[LayerIdx] <= OutputBufferSize);
+
 					Layers[LayerIdx]->Evaluate(
 						ConcatInstance->Instances[LayerIdx].Get(),
 						OutputBuffer + OutputOffsets[LayerIdx],
@@ -2654,6 +2736,7 @@ namespace UE::NNE::RuntimeBasic
 			{
 				NNE_RUNTIME_BASIC_TRACE_SCOPE(NNE::RuntimeBasic::Private::FArrayLayer::Evaluate);
 				check(OutputBufferSize == GetOutputSize() && InputBufferSize == GetInputSize());
+				check(OutputBufferStride >= GetOutputSize() && InputBufferStride >= GetInputSize());
 
 				FArrayLayerInstance* ArrayInstance = StaticCast<FArrayLayerInstance*>(Instance);
 				check(ArrayInstance);
@@ -2829,6 +2912,7 @@ namespace UE::NNE::RuntimeBasic
 			{
 				NNE_RUNTIME_BASIC_TRACE_SCOPE(NNE::RuntimeBasic::Private::FAggregateSetLayer::Evaluate);
 				check(OutputBufferSize == GetOutputSize() && InputBufferSize == GetInputSize());
+				check(OutputBufferStride >= GetOutputSize() && InputBufferStride >= GetInputSize());
 
 				FAggregateSetLayerInstance* AggregateSetInstance = StaticCast<FAggregateSetLayerInstance*>(Instance);
 				check(AggregateSetInstance);
@@ -3075,6 +3159,7 @@ namespace UE::NNE::RuntimeBasic
 			{
 				NNE_RUNTIME_BASIC_TRACE_SCOPE(NNE::RuntimeBasic::Private::FAggregateOrExclusiveLayer::Evaluate);
 				check(OutputBufferSize == GetOutputSize() && InputBufferSize == GetInputSize());
+				check(OutputBufferStride >= GetOutputSize() && InputBufferStride >= GetInputSize());
 
 				FAggregateOrExclusiveLayerInstance* AggregateOrExclusiveInstance = StaticCast<FAggregateOrExclusiveLayerInstance*>(Instance);
 				check(AggregateOrExclusiveInstance);
@@ -3314,6 +3399,7 @@ namespace UE::NNE::RuntimeBasic
 			{
 				NNE_RUNTIME_BASIC_TRACE_SCOPE(NNE::RuntimeBasic::Private::FAggregateOrInclusiveLayer::Evaluate);
 				check(OutputBufferSize == GetOutputSize() && InputBufferSize == GetInputSize());
+				check(OutputBufferStride >= GetOutputSize() && InputBufferStride >= GetInputSize());
 
 				FAggregateOrInclusiveLayerInstance* AggregateOrInclusiveInstance = StaticCast<FAggregateOrInclusiveLayerInstance*>(Instance);
 				check(AggregateOrInclusiveInstance);
@@ -3578,6 +3664,7 @@ namespace UE::NNE::RuntimeBasic
 			{
 				NNE_RUNTIME_BASIC_TRACE_SCOPE(NNE::RuntimeBasic::Private::FClampLayer::Evaluate);
 				check(OutputBufferSize == GetOutputSize() && InputBufferSize == GetInputSize());
+				check(OutputBufferStride >= GetOutputSize() && InputBufferStride >= GetInputSize());
 				check(Instance == nullptr);
 				OperatorNanCheck(InputBuffer, BatchSize, InputBufferSize, InputBufferStride);
 
@@ -3673,6 +3760,7 @@ namespace UE::NNE::RuntimeBasic
 			{
 				NNE_RUNTIME_BASIC_TRACE_SCOPE(NNE::RuntimeBasic::Private::FSparseMixtureOfExpertsLayer::Evaluate);
 				check(OutputBufferSize == GetOutputSize() && InputBufferSize == GetInputSize());
+				check(OutputBufferStride >= GetOutputSize() && InputBufferStride >= GetInputSize());
 
 				FSparseMixtureOfExpertsLayerInstance* SparseMixtureOfExpertsInstance = StaticCast<FSparseMixtureOfExpertsLayerInstance*>(Instance);
 				check(SparseMixtureOfExpertsInstance);
