[NNE] NNERuntimeBasicCPU unify log across NNE runtimes

Problem:
* Log have been moved to a per runtime category in other NNE runtimes.

Solution:
* Unify by adding a log category to NNERuntimeBasicCPU.

Tests:
* Build win64 Editor debug & win64 Standalone packaged debug

#rnx

diff --git a/Engine/Plugins/Experimental/NNERuntimeBasicCpu/Source/NNERuntimeBasicCpu/Private/NNERuntimeBasicCpu.cpp b/Engine/Plugins/Experimental/NNERuntimeBasicCpu/Source/NNERuntimeBasicCpu/Private/NNERuntimeBasicCpu.cpp
--- a/Engine/Plugins/Experimental/NNERuntimeBasicCpu/Source/NNERuntimeBasicCpu/Private/NNERuntimeBasicCpu.cpp
+++ b/Engine/Plugins/Experimental/NNERuntimeBasicCpu/Source/NNERuntimeBasicCpu/Private/NNERuntimeBasicCpu.cpp
@@ -6,6 +6,8 @@
 #include "NNERuntimeBasicCpuModel.h"
 #include "NNEModelData.h"
 
+DEFINE_LOG_CATEGORY(LogNNERuntimeBasicCPU);
+
 // We ask for the memory to be aligned to 64 bytes since this is the
 // largest alignment we ask for inside the ModelData for pointers to
 // various bits of data.
@@ -57,7 +59,7 @@ TSharedPtr<UE::NNE::FSharedModelData> UNNERuntimeBasicCpuImpl::CreateModelData(c
 {
 	if (CanCreateModelData(FileType, FileData, AdditionalFileData, FileId, TargetPlatform) != ECanCreateModelDataStatus::Ok)
 	{
-		UE_LOG(LogNNE, Warning, TEXT("UNNERuntimeBasicCpu cannot create the model data with id %s (Filetype: %s)"), *FileId.ToString(EGuidFormats::Digits).ToLower(), *FileType);
+		UE_LOG(LogNNERuntimeBasicCPU, Warning, TEXT("Cannot create the model data with id %s (Filetype: %s)"), *FileId.ToString(EGuidFormats::Digits).ToLower(), *FileType);
 		return nullptr;
 	}
 
@@ -118,7 +120,7 @@ TSharedPtr<UE::NNE::IModelCPU> UNNERuntimeBasicCpuImpl::CreateModelCPU(const TOb
 
 	if (CanCreateModelCPU(ModelData) != ECanCreateModelCPUStatus::Ok)
 	{
-		UE_LOG(LogNNE, Warning, TEXT("UNNERuntimeBasicCpu cannot create a model from the model data with id %s"), *ModelData->GetFileId().ToString(EGuidFormats::Digits));
+		UE_LOG(LogNNERuntimeBasicCPU, Warning, TEXT("Cannot create a model from the model data with id %s"), *ModelData->GetFileId().ToString(EGuidFormats::Digits));
 		return nullptr;
 	}
 
diff --git a/Engine/Plugins/Experimental/NNERuntimeBasicCpu/Source/NNERuntimeBasicCpu/Private/NNERuntimeBasicCpu.h b/Engine/Plugins/Experimental/NNERuntimeBasicCpu/Source/NNERuntimeBasicCpu/Private/NNERuntimeBasicCpu.h
--- a/Engine/Plugins/Experimental/NNERuntimeBasicCpu/Source/NNERuntimeBasicCpu/Private/NNERuntimeBasicCpu.h
+++ b/Engine/Plugins/Experimental/NNERuntimeBasicCpu/Source/NNERuntimeBasicCpu/Private/NNERuntimeBasicCpu.h
@@ -2,10 +2,13 @@
 
 #pragma once
 
+#include "Logging/LogMacros.h"
 #include "NNERuntime.h"
 #include "NNERuntimeCPU.h"
 #include "UObject/Object.h"
 
+DECLARE_LOG_CATEGORY_EXTERN(LogNNERuntimeBasicCPU, Log, All);
+
 #include "NNERuntimeBasicCpu.generated.h"
 
 /**
diff --git a/Engine/Plugins/Experimental/NNERuntimeBasicCpu/Source/NNERuntimeBasicCpu/Private/NNERuntimeBasicCpuModel.cpp b/Engine/Plugins/Experimental/NNERuntimeBasicCpu/Source/NNERuntimeBasicCpu/Private/NNERuntimeBasicCpuModel.cpp
--- a/Engine/Plugins/Experimental/NNERuntimeBasicCpu/Source/NNERuntimeBasicCpu/Private/NNERuntimeBasicCpuModel.cpp
+++ b/Engine/Plugins/Experimental/NNERuntimeBasicCpu/Source/NNERuntimeBasicCpu/Private/NNERuntimeBasicCpuModel.cpp
@@ -4098,7 +4098,7 @@ namespace UE::NNE::RuntimeBasic
 		Private::Serialization::Load(InOutOffset, Magic, Data);
 		if (Magic != ModelMagicNumber)
 		{
-			UE_LOG(LogNNE, Error, TEXT("Invalid Magic Number %i"), Magic);
+			UE_LOG(LogNNERuntimeBasicCPU, Error, TEXT("Invalid Magic Number %i"), Magic);
 			return false;
 		}
 
@@ -4106,7 +4106,7 @@ namespace UE::NNE::RuntimeBasic
 		Private::Serialization::Load(InOutOffset, Version, Data);
 		if (Version != ModelVersionNumber)
 		{
-			UE_LOG(LogNNE, Error, TEXT("Unsupported Version Number %i"), Version);
+			UE_LOG(LogNNERuntimeBasicCPU, Error, TEXT("Unsupported Version Number %i"), Version);
 			return false;
 		}
 
