NNERuntimeBasicCpu: Added LayerNorm and Lipschizt Linear Layers

diff --git a/Engine/Plugins/Experimental/NNERuntimeBasicCpu/Source/NNERuntimeBasicCpu/Private/NNERuntimeBasicCpu.ispc b/Engine/Plugins/Experimental/NNERuntimeBasicCpu/Source/NNERuntimeBasicCpu/Private/NNERuntimeBasicCpu.ispc
--- a/Engine/Plugins/Experimental/NNERuntimeBasicCpu/Source/NNERuntimeBasicCpu/Private/NNERuntimeBasicCpu.ispc
+++ b/Engine/Plugins/Experimental/NNERuntimeBasicCpu/Source/NNERuntimeBasicCpu/Private/NNERuntimeBasicCpu.ispc
@@ -16,6 +16,11 @@ inline float NNERuntimeBasicCPUGELU(const float X)
 	return X * NNERuntimeBasicCPUSigmoid(1.702f * X);
 }
 
+inline float NNERuntimeBasicCPUSquare(const float X)
+{
+	return X * X;
+}
+
 export void NNERuntimeBasicCPUOperatorNormalize(
 	uniform float Output[],
 	uniform const float Input[],
@@ -555,4 +560,38 @@ export void NNERuntimeBasicCPUOperatorScatter(
 			OutputBuffer[DstBatchIdx * OutputStride + Idx] = InputBuffer[BatchIndexIdx * InputStride + Idx];
 		}
 	}
+}
+
+export void NNERuntimeBasicCPUOperatorLayerNorm(
+	uniform float Output[],
+	uniform const float Input[],
+	uniform const float Offset[],
+	uniform const float Scale[],
+	uniform const float Epsilon,
+	uniform const int BatchNum,
+	uniform const int InputOutputNum,
+	uniform const int OutputStride,
+	uniform const int InputStride)
+{
+	for (uniform int BatchIdx = 0; BatchIdx < BatchNum; BatchIdx++)
+	{
+		float MeanAccum = 0.0f;
+		foreach (Idx = 0 ... InputOutputNum)
+		{
+			MeanAccum += Input[BatchIdx * InputStride + Idx] / InputOutputNum;
+		}
+		uniform float Mean = reduce_add(MeanAccum);
+
+		float StdAccum = 0.0f;
+		foreach (Idx = 0 ... InputOutputNum)
+		{
+			StdAccum += NNERuntimeBasicCPUSquare(Input[BatchIdx * InputStride + Idx] - Mean) / InputOutputNum;
+		}
+		uniform float Std = sqrt(reduce_add(StdAccum) + Epsilon);
+
+		foreach (Idx = 0 ... InputOutputNum)
+		{
+			Output[BatchIdx * OutputStride + Idx] = ((Input[BatchIdx * InputStride + Idx] - Mean) / Std) * Scale[Idx] + Offset[Idx];
+		}
+	}
 }
\ No newline at end of file
diff --git a/Engine/Plugins/Experimental/NNERuntimeBasicCpu/Source/NNERuntimeBasicCpu/Private/NNERuntimeBasicCpuModel.cpp b/Engine/Plugins/Experimental/NNERuntimeBasicCpu/Source/NNERuntimeBasicCpu/Private/NNERuntimeBasicCpuModel.cpp
--- a/Engine/Plugins/Experimental/NNERuntimeBasicCpu/Source/NNERuntimeBasicCpu/Private/NNERuntimeBasicCpuModel.cpp
+++ b/Engine/Plugins/Experimental/NNERuntimeBasicCpu/Source/NNERuntimeBasicCpu/Private/NNERuntimeBasicCpuModel.cpp
@@ -1349,6 +1349,56 @@ namespace UE::NNE::RuntimeBasic
 			}
 		}
 
+		static inline void OperatorLayerNorm(
+			float* RESTRICT Output,
+			const float* RESTRICT Input,
+			const float* RESTRICT Offset,
+			const float* RESTRICT Scale,
+			const float Epsilon,
+			const uint32 BatchSize,
+			const uint32 InputOutputSize,
+			const uint32 OutputStride,
+			const uint32 InputStride)
+		{
+			NNE_RUNTIME_BASIC_TRACE_SCOPE(NNE::RuntimeBasic::Private::OperatorLayerNorm);
+
+			check(Output != Input);
+
+#if NNE_RUNTIME_BASIC_ENABLE_ISPC
+			ispc::NNERuntimeBasicCPUOperatorLayerNorm(
+				Output,
+				Input,
+				Offset,
+				Scale,
+				Epsilon,
+				BatchSize,
+				InputOutputSize,
+				OutputStride,
+				InputStride);
+#else
+			for (uint32 BatchIdx = 0; BatchIdx < BatchSize; BatchIdx++)
+			{
+				float Mean = 0.0f;
+				for (uint32 Idx = 0; Idx < InputOutputSize; Idx++)
+				{
+					Mean += Input[BatchIdx * InputStride + Idx] / InputOutputSize;
+				}
+
+				float Std = 0.0f;
+				for (uint32 Idx = 0; Idx < InputOutputSize; Idx++)
+				{
+					Std += FMath::Square(Input[BatchIdx * InputStride + Idx] - Mean) / InputOutputSize;
+				}
+				Std = FMath::Sqrt(Std + Epsilon);
+
+				for (uint32 Idx = 0; Idx < InputOutputSize; Idx++)
+				{
+					Output[BatchIdx * OutputStride + Idx] = ((Input[BatchIdx * InputStride + Idx] - Mean) / Std) * Scale[Idx] + Offset[Idx];
+				}
+			}
+#endif
+		}
+
 		//--------------------------------------------------------------------------
 		// Layer Types
 		//--------------------------------------------------------------------------
@@ -1357,33 +1407,28 @@ namespace UE::NNE::RuntimeBasic
 		enum class ELayerType : uint32
 		{
 			Invalid = 0,
-			
 			Sequence = 1,
-			
 			Normalize = 2,
 			Denormalize = 3,
 			Linear = 4,
 			CompressedLinear = 5,
 			MultiLinear = 6,
-
 			ReLU = 7,
 			ELU = 8,
 			TanH = 9,
 			PReLU = 10,
-
 			MemoryCell = 11,
-
 			Copy = 12,
 			Concat = 13,
 			Array = 14,
-			
 			AggregateSet = 15,
 			AggregateOrExclusive = 16,
 			AggregateOrInclusive = 17,
-
 			Clamp = 18,
 			SparseMixtureOfExperts = 19,
 			GELU = 20,
+			LayerNorm = 21,
+			LipschiztLinear = 22,
 		};
 
 		//--------------------------------------------------------------------------
@@ -3898,6 +3943,142 @@ namespace UE::NNE::RuntimeBasic
 			}
 		}
 
+		//--------------------------------------------------------------------------
+
+		struct FLayerNormLayer : public ILayer
+		{
+			virtual ELayerType GetLayerType() const override final { return ELayerType::LayerNorm; }
+			virtual uint32 GetInputSize() const override final { return InputOutputSize; }
+			virtual uint32 GetOutputSize() const override final { return InputOutputSize; }
+
+			virtual void SerializationSize(uint64& InOutOffset) const override final
+			{
+				Serialization::Size(InOutOffset, InputOutputSize);
+				Serialization::Size(InOutOffset, Offset);
+				Serialization::Size(InOutOffset, Scale);
+				Serialization::Size(InOutOffset, Epsilon);
+			}
+
+			virtual void SerializationLoad(uint64& InOutOffset, TConstArrayView<uint8> Data) override final
+			{
+				Serialization::Load(InOutOffset, InputOutputSize, Data);
+				Serialization::Load(InOutOffset, Offset, Data, InputOutputSize);
+				Serialization::Load(InOutOffset, Scale, Data, InputOutputSize);
+				Serialization::Load(InOutOffset, Epsilon, Data);
+			}
+
+			virtual void SerializationSave(uint64& InOutOffset, TArrayView<uint8> Data) const override final
+			{
+				Serialization::Save(InOutOffset, InputOutputSize, Data);
+				Serialization::Save(InOutOffset, Offset, Data);
+				Serialization::Save(InOutOffset, Scale, Data);
+				Serialization::Save(InOutOffset, Epsilon, Data);
+			}
+
+			virtual void Evaluate(
+				ILayerInstance* Instance,
+				float* OutputBuffer,
+				const float* InputBuffer,
+				const uint32 BatchSize,
+				const uint32 OutputBufferSize,
+				const uint32 InputBufferSize,
+				const uint32 OutputBufferStride,
+				const uint32 InputBufferStride) override final
+			{
+				NNE_RUNTIME_BASIC_TRACE_SCOPE(NNE::RuntimeBasic::Private::FLayerNormLayer::Evaluate);
+				check(OutputBufferSize == GetOutputSize() && InputBufferSize == GetInputSize());
+				check(OutputBufferStride >= GetOutputSize() && InputBufferStride >= GetInputSize());
+				check(Instance == nullptr);
+				OperatorNanCheck(InputBuffer, BatchSize, InputBufferSize, InputBufferStride);
+
+				OperatorLayerNorm(
+					OutputBuffer,
+					InputBuffer,
+					Offset.GetData(),
+					Scale.GetData(),
+					Epsilon,
+					BatchSize,
+					InputOutputSize,
+					OutputBufferStride,
+					InputBufferStride);
+
+				OperatorNanCheck(OutputBuffer, BatchSize, OutputBufferSize, OutputBufferStride);
+			}
+
+			uint32 InputOutputSize = 0;
+			TConstArrayView<float> Offset;
+			TConstArrayView<float> Scale;
+			float Epsilon = 1e-5f;
+		};
+
+		//--------------------------------------------------------------------------
+
+		struct FLipschiztLinearLayer : public ILayer
+		{
+			virtual ELayerType GetLayerType() const override final { return ELayerType::LipschiztLinear; }
+			virtual uint32 GetInputSize() const override final { return InputSize; }
+			virtual uint32 GetOutputSize() const override final { return OutputSize; }
+
+			virtual void SerializationSize(uint64& InOutOffset) const override final
+			{
+				Serialization::Size(InOutOffset, InputSize);
+				Serialization::Size(InOutOffset, OutputSize);
+				Serialization::Size(InOutOffset, Biases);
+				Serialization::Size(InOutOffset, Weights);
+			}
+
+			virtual void SerializationLoad(uint64& InOutOffset, TConstArrayView<uint8> Data) override final
+			{
+				Serialization::Load(InOutOffset, InputSize, Data);
+				Serialization::Load(InOutOffset, OutputSize, Data);
+				Serialization::Load(InOutOffset, Biases, Data, OutputSize);
+				Serialization::Load(InOutOffset, Weights, Data, InputSize * OutputSize);
+			}
+
+			virtual void SerializationSave(uint64& InOutOffset, TArrayView<uint8> Data) const override final
+			{
+				Serialization::Save(InOutOffset, InputSize, Data);
+				Serialization::Save(InOutOffset, OutputSize, Data);
+				Serialization::Save(InOutOffset, Biases, Data);
+				Serialization::Save(InOutOffset, Weights, Data);
+			}
+
+			virtual void Evaluate(
+				ILayerInstance* Instance,
+				float* OutputBuffer,
+				const float* InputBuffer,
+				const uint32 BatchSize,
+				const uint32 OutputBufferSize,
+				const uint32 InputBufferSize,
+				const uint32 OutputBufferStride,
+				const uint32 InputBufferStride) override final
+			{
+				NNE_RUNTIME_BASIC_TRACE_SCOPE(NNE::RuntimeBasic::Private::FLipschiztLinearLayer::Evaluate);
+				check(OutputBufferSize == GetOutputSize() && InputBufferSize == GetInputSize());
+				check(OutputBufferStride >= GetOutputSize() && InputBufferStride >= GetInputSize());
+				check(Instance == nullptr);
+				OperatorNanCheck(InputBuffer, BatchSize, InputBufferSize, InputBufferStride);
+
+				OperatorLinear(
+					OutputBuffer,
+					InputBuffer,
+					Weights.GetData(),
+					Biases.GetData(),
+					BatchSize,
+					OutputSize,
+					InputSize,
+					OutputBufferStride,
+					InputBufferStride);
+
+				OperatorNanCheck(OutputBuffer, BatchSize, OutputBufferSize, OutputBufferStride);
+			}
+
+			uint32 InputSize = 0;
+			uint32 OutputSize = 0;
+			TConstArrayView<float> Biases;
+			TConstArrayView<float> Weights;
+		};
+
 		//--------------------------------------------------------------------------
 		// Layer Serialization
 		//--------------------------------------------------------------------------
@@ -3946,6 +4127,8 @@ namespace UE::NNE::RuntimeBasic
 					case ELayerType::Clamp: OutLayer = MakeShared<FClampLayer>(); break;
 					case ELayerType::SparseMixtureOfExperts: OutLayer = MakeShared<FSparseMixtureOfExpertsLayer>(); break;
 					case ELayerType::GELU: OutLayer = MakeShared<FGELULayer>(); break;
+					case ELayerType::LayerNorm: OutLayer = MakeShared<FLayerNormLayer>(); break;
+					case ELayerType::LipschiztLinear: OutLayer = MakeShared<FLipschiztLinearLayer>(); break;
 					default: checkf(false, TEXT("Unknown Layer Id %i"), LayerTypeId);
 					}
 				}
@@ -4713,6 +4896,42 @@ namespace UE::NNE::RuntimeBasic
 		return StaticCastSharedPtr<Private::ILayer>(SparseMixtureOfExpertsLayer);
 	}
 
+	FModelBuilderElement FModelBuilder::MakeLayerNorm(
+		const uint32 InputOutputSize,
+		const TConstArrayView<float> Offsets,
+		const TConstArrayView<float> Scales,
+		const float Epsilon)
+	{
+		check(Offsets.Num() == InputOutputSize);
+		check(Scales.Num() == InputOutputSize);
+
+		const TSharedPtr<Private::FLayerNormLayer> LayerNormLayer = MakeShared<Private::FLayerNormLayer>();
+		LayerNormLayer->InputOutputSize = InputOutputSize;
+		LayerNormLayer->Offset = Offsets;
+		LayerNormLayer->Scale = Scales;
+		LayerNormLayer->Epsilon = Epsilon;
+
+		return StaticCastSharedPtr<Private::ILayer>(LayerNormLayer);
+	}
+
+	FModelBuilderElement FModelBuilder::MakeLipschiztLinear(
+		const uint32 InputSize,
+		const uint32 OutputSize,
+		const TConstArrayView<float> Weights,
+		const TConstArrayView<float> Biases)
+	{
+		check(Biases.Num() == OutputSize);
+		check(Weights.Num() == InputSize * OutputSize);
+
+		const TSharedPtr<Private::FLipschiztLinearLayer> LipschiztLinearLayer = MakeShared<Private::FLipschiztLinearLayer>();
+		LipschiztLinearLayer->InputSize = InputSize;
+		LipschiztLinearLayer->OutputSize = OutputSize;
+		LipschiztLinearLayer->Biases = Biases;
+		LipschiztLinearLayer->Weights = Weights;
+
+		return StaticCastSharedPtr<Private::ILayer>(LipschiztLinearLayer);
+	}
+
 	void FModelBuilder::Reset()
 	{
 		Rng = RngInitialState;
diff --git a/Engine/Plugins/Experimental/NNERuntimeBasicCpu/Source/NNERuntimeBasicCpu/Public/NNERuntimeBasicCpuBuilder.h b/Engine/Plugins/Experimental/NNERuntimeBasicCpu/Source/NNERuntimeBasicCpu/Public/NNERuntimeBasicCpuBuilder.h
--- a/Engine/Plugins/Experimental/NNERuntimeBasicCpu/Source/NNERuntimeBasicCpu/Public/NNERuntimeBasicCpuBuilder.h
+++ b/Engine/Plugins/Experimental/NNERuntimeBasicCpu/Source/NNERuntimeBasicCpu/Public/NNERuntimeBasicCpuBuilder.h
@@ -346,6 +346,34 @@ namespace UE::NNE::RuntimeBasic
 			const FModelBuilderElement& GatingLayer,
 			const TConstArrayView<FModelBuilderElement> SubLayers);
 
+		/**
+		 * Makes a new Layer Norm layer.
+		 *
+		 * @param InputOutputSize	Input and Output Vector Size
+		 * @param Offsets			Initial Offsets
+		 * @param Scales			Initial Scales
+		 * @param Epsilon			Standard Deviation Epsilon
+		 */
+		FModelBuilderElement MakeLayerNorm(
+			const uint32 InputOutputSize,
+			const TConstArrayView<float> Offsets,
+			const TConstArrayView<float> Scales,
+			const float Epsilon = 1e-5f);
+
+		/**
+		 * Makes a new Lipschizt linear layer.
+		 *
+		 * @param InputSize		Input Vector Size (Number of Rows)
+		 * @param OutputSize	Output Vector Size (Number of Columns)
+		 * @param Weights		Linear layer weights.
+		 * @param Biases		Linear layer biases.
+		 */
+		FModelBuilderElement MakeLipschiztLinear(
+			const uint32 InputSize,
+			const uint32 OutputSize,
+			const TConstArrayView<float> Weights,
+			const TConstArrayView<float> Biases);
+
 	public:
 
 		/** Creates a array of weights from a copy of the given array view */
