NNERuntimeBasicCpu: Added MLP with LayerNorm.

diff --git a/Engine/Plugins/Experimental/NNERuntimeBasicCpu/Source/NNERuntimeBasicCpu/Private/NNERuntimeBasicCpuModel.cpp b/Engine/Plugins/Experimental/NNERuntimeBasicCpu/Source/NNERuntimeBasicCpu/Private/NNERuntimeBasicCpuModel.cpp
--- a/Engine/Plugins/Experimental/NNERuntimeBasicCpu/Source/NNERuntimeBasicCpu/Private/NNERuntimeBasicCpuModel.cpp
+++ b/Engine/Plugins/Experimental/NNERuntimeBasicCpu/Source/NNERuntimeBasicCpu/Private/NNERuntimeBasicCpuModel.cpp
@@ -4605,6 +4605,42 @@ namespace UE::NNE::RuntimeBasic
 		return MakeSequence(Layers);
 	}
 
+	FModelBuilderElement FModelBuilder::MakeMLPWithRandomKaimingWeightsWithLayerNorm(
+		const uint32 InputSize, 
+		const uint32 OutputSize, 
+		const uint32 HiddenSize, 
+		const uint32 LayerNum, 
+		const EActivationFunction ActivationFunction, 
+		const bool bActivationOnFinalLayer)
+	{
+		check(LayerNum >= 2);
+
+		TArray<FModelBuilderElement, TInlineAllocator<32>> Layers;
+		Layers.Reserve(2 * LayerNum - (bActivationOnFinalLayer ? 0 : 1) + LayerNum - 1);
+
+		for (uint32 LayerIdx = 0; LayerIdx < LayerNum; LayerIdx++)
+		{
+			const uint32 LayerInputSize = LayerIdx == 0 ? InputSize : HiddenSize;
+			const uint32 LayerOuputSize = LayerIdx == LayerNum - 1 ? OutputSize : HiddenSize;
+
+			Layers.Emplace(MakeLinearWithRandomKaimingWeights(LayerInputSize, LayerOuputSize));
+			
+			if (LayerIdx != LayerNum - 1)
+			{
+				Layers.Emplace(
+					MakeLayerNorm(LayerOuputSize, MakeWeightsZero(LayerOuputSize), MakeWeightsConstant(LayerOuputSize, 1.0f))
+				);
+			}
+
+			if (bActivationOnFinalLayer || LayerIdx != LayerNum - 1)
+			{
+				Layers.Emplace(MakeActivation(LayerOuputSize, ActivationFunction));
+			}
+		}
+
+		return MakeSequence(Layers);
+	}
+
 	FModelBuilderElement FModelBuilder::MakeCompressedMLPWithRandomKaimingWeights(
 		const uint32 InputSize,
 		const uint32 OutputSize,
diff --git a/Engine/Plugins/Experimental/NNERuntimeBasicCpu/Source/NNERuntimeBasicCpu/Public/NNERuntimeBasicCpuBuilder.h b/Engine/Plugins/Experimental/NNERuntimeBasicCpu/Source/NNERuntimeBasicCpu/Public/NNERuntimeBasicCpuBuilder.h
--- a/Engine/Plugins/Experimental/NNERuntimeBasicCpu/Source/NNERuntimeBasicCpu/Public/NNERuntimeBasicCpuBuilder.h
+++ b/Engine/Plugins/Experimental/NNERuntimeBasicCpu/Source/NNERuntimeBasicCpu/Public/NNERuntimeBasicCpuBuilder.h
@@ -49,6 +49,7 @@ namespace UE::NNE::RuntimeBasic
 		/** Common activation function types. */
 		enum class EActivationFunction : uint8
 		{
+
 			ReLU = 0,
 			ELU = 1,
 			TanH = 2,
@@ -203,6 +204,25 @@ namespace UE::NNE::RuntimeBasic
 			const EActivationFunction ActivationFunction,
 			const bool bActivationOnFinalLayer = false);
 
+
+		/**
+		 * Makes a Multi-Layer Perceptron network with randomly initialized weights using the Kaiming method with LayerNorm before all activations.
+		 *
+		 * @param InputSize					Input Vector Size
+		 * @param OutputSize				Output Vector Size
+		 * @param HiddenSize				Number of hidden units to use on internal layers.
+		 * @param LayerNum					Number of layers. Includes input and output layers.
+		 * @param ActivationFunction		Activation function to use.
+		 * @param bActivationOnFinalLayer	If the activation function should be used on the final output layer.
+		 */
+		FModelBuilderElement MakeMLPWithRandomKaimingWeightsWithLayerNorm(
+			const uint32 InputSize,
+			const uint32 OutputSize,
+			const uint32 HiddenSize,
+			const uint32 LayerNum,
+			const EActivationFunction ActivationFunction,
+			const bool bActivationOnFinalLayer = false);
+
 		/**
 		 * Makes a Compressed Multi-Layer Perceptron network with randomly initialized weights using the Kaiming method.
 		 *
@@ -475,5 +495,4 @@ namespace UE::NNE::RuntimeBasic
 		TArray<TArray<uint32>> SizesPool;
 	};
 
-
 }
\ No newline at end of file
