NNERuntimeBasicCpu: Refactored and cleaned up FModelBuilder a little to reduce the number of function variations and make it easier to use different kinds of Linear Layer and Weight Initialization inside of more complex models by passing these details in as settings. Added support for Kaiming Uniform weight initialization which should more accurately match the PyTorch default weight initialization. Added clipping to the current default gaussian weight initialization to avoid any chance of large outliers.

diff --git a/Engine/Plugins/Animation/MLDeformer/NearestNeighborModel/Source/NearestNeighborModel/Private/NearestNeighborOptimizedNetwork.cpp b/Engine/Plugins/Animation/MLDeformer/NearestNeighborModel/Source/NearestNeighborModel/Private/NearestNeighborOptimizedNetwork.cpp
--- a/Engine/Plugins/Animation/MLDeformer/NearestNeighborModel/Source/NearestNeighborModel/Private/NearestNeighborOptimizedNetwork.cpp
+++ b/Engine/Plugins/Animation/MLDeformer/NearestNeighborModel/Source/NearestNeighborModel/Private/NearestNeighborOptimizedNetwork.cpp
@@ -37,7 +37,7 @@ namespace UE::NearestNeighborModel::Private
 					GemmPreluLayer->Parameters[0].Values,
 					GemmPreluLayer->Parameters[1].Values));
 
-				LayerElements.Add(Builder.MakePReLU(GemmPreluLayer->NumOutputs, Builder.MakeWeightsConstant(GemmPreluLayer->NumOutputs, GemmPreluLayer->Parameters[2].Values[0])));
+				LayerElements.Add(Builder.MakePReLU(GemmPreluLayer->NumOutputs, Builder.MakeValuesConstant(GemmPreluLayer->NumOutputs, GemmPreluLayer->Parameters[2].Values[0])));
 			}
 			else if (UNearestNeighborNetworkLayer_Gemm* GemmLayer = Cast<UNearestNeighborNetworkLayer_Gemm>(Layer))
 			{
diff --git a/Engine/Plugins/Experimental/LearningAgents/Source/Learning/Private/LearningAction.cpp b/Engine/Plugins/Experimental/LearningAgents/Source/Learning/Private/LearningAction.cpp
--- a/Engine/Plugins/Experimental/LearningAgents/Source/Learning/Private/LearningAction.cpp
+++ b/Engine/Plugins/Experimental/LearningAgents/Source/Learning/Private/LearningAction.cpp
@@ -1072,8 +1072,8 @@ namespace UE::Learning::Action
 
 			OutElement = Builder.MakeDenormalize(
 				ValueNum,
-				Builder.MakeWeightsZero(ValueNum),
-				Builder.MakeWeightsConstant(ValueNum, 1.0f));
+				Builder.MakeValuesZero(ValueNum),
+				Builder.MakeValuesOne(ValueNum));
 			break;
 		}
 
@@ -1091,8 +1091,8 @@ namespace UE::Learning::Action
 
 			OutElement = Builder.MakeDenormalize(
 				Parameters.Num,
-				Builder.MakeWeightsCopy(LogPriorProbabilities),
-				Builder.MakeWeightsConstant(Parameters.Num, 1.0f));
+				Builder.MakeValuesCopy(LogPriorProbabilities),
+				Builder.MakeValuesOne(Parameters.Num));
 
 			break;
 		}
@@ -1109,8 +1109,8 @@ namespace UE::Learning::Action
 			}
 
 			OutElement = Builder.MakeDenormalize(Parameters.Num,
-				Builder.MakeWeightsCopy(LogPriorProbabilities),
-				Builder.MakeWeightsConstant(Parameters.Num, 1.0f));
+				Builder.MakeValuesCopy(LogPriorProbabilities),
+				Builder.MakeValuesOne(Parameters.Num));
 			break;
 		}
 
@@ -1154,8 +1154,8 @@ namespace UE::Learning::Action
 
 			BuilderLayers.Emplace(Builder.MakeDenormalize(
 				LogPriorProbabilities.Num(),
-				Builder.MakeWeightsCopy(LogPriorProbabilities),
-				Builder.MakeWeightsConstant(LogPriorProbabilities.Num(), 1.0f)));
+				Builder.MakeValuesCopy(LogPriorProbabilities),
+				Builder.MakeValuesOne(LogPriorProbabilities.Num())));
 
 			OutElement = Builder.MakeConcat(BuilderLayers);
 			break;
@@ -1183,8 +1183,8 @@ namespace UE::Learning::Action
 
 			BuilderLayers.Emplace(Builder.MakeDenormalize(
 				LogPriorProbabilities.Num(),
-				Builder.MakeWeightsCopy(LogPriorProbabilities),
-				Builder.MakeWeightsConstant(LogPriorProbabilities.Num(), 1.0f)));
+				Builder.MakeValuesCopy(LogPriorProbabilities),
+				Builder.MakeValuesOne(LogPriorProbabilities.Num())));
 
 			OutElement = Builder.MakeConcat(BuilderLayers);
 			break;
@@ -1209,34 +1209,23 @@ namespace UE::Learning::Action
 			NNE::RuntimeBasic::FModelBuilderElement BuilderSubElement;
 			MakeDecoderNetworkModelBuilderElementFromSchema(BuilderSubElement, Builder, Schema, Parameters.Element, NetworkSettings);
 
-			if (NetworkSettings.bUseCompressedLinearLayers)
-			{
-				OutElement = Builder.MakeSequence({
-					Builder.MakeActivation(Parameters.EncodingSize, Private::GetNNEActivationFunction(Parameters.ActivationFunction)),
-					Builder.MakeCompressedMLPWithRandomKaimingWeights(
-						Parameters.EncodingSize,
-						SubElementEncodedSize,
-						Parameters.EncodingSize,
-						Parameters.LayerNum + 1,  // Add 1 to account for input layer 
-						Private::GetNNEActivationFunction(Parameters.ActivationFunction),
-						false),
-					BuilderSubElement,
-					});
-			}
-			else
-			{
-				OutElement = Builder.MakeSequence({
-					Builder.MakeActivation(Parameters.EncodingSize, Private::GetNNEActivationFunction(Parameters.ActivationFunction)),
-					Builder.MakeMLPWithRandomKaimingWeights(
-						Parameters.EncodingSize,
-						SubElementEncodedSize,
-						Parameters.EncodingSize,
-						Parameters.LayerNum + 1,  // Add 1 to account for input layer 
-						Private::GetNNEActivationFunction(Parameters.ActivationFunction),
-						false),
-					BuilderSubElement,
-					});
-			}
+			NNE::RuntimeBasic::FModelBuilder::FLinearLayerSettings LinearLayerSettings;
+			LinearLayerSettings.Type = NetworkSettings.bUseCompressedLinearLayers ?
+				NNE::RuntimeBasic::FModelBuilder::ELinearLayerType::Compressed :
+				NNE::RuntimeBasic::FModelBuilder::ELinearLayerType::Normal;
+
+			OutElement = Builder.MakeSequence({
+				Builder.MakeActivation(Parameters.EncodingSize, Private::GetNNEActivationFunction(Parameters.ActivationFunction)),
+				Builder.MakeMLP(
+					Parameters.EncodingSize,
+					SubElementEncodedSize,
+					Parameters.EncodingSize,
+					Parameters.LayerNum + 1,  // Add 1 to account for input layer 
+					Private::GetNNEActivationFunction(Parameters.ActivationFunction),
+					false,
+					LinearLayerSettings),
+				BuilderSubElement,
+				});
 
 			break;
 		}
diff --git a/Engine/Plugins/Experimental/LearningAgents/Source/Learning/Private/LearningObservation.cpp b/Engine/Plugins/Experimental/LearningAgents/Source/Learning/Private/LearningObservation.cpp
--- a/Engine/Plugins/Experimental/LearningAgents/Source/Learning/Private/LearningObservation.cpp
+++ b/Engine/Plugins/Experimental/LearningAgents/Source/Learning/Private/LearningObservation.cpp
@@ -929,8 +929,8 @@ namespace UE::Learning::Observation
 
 			OutElement = Builder.MakeDenormalize(
 				ValueNum,
-				Builder.MakeWeightsZero(ValueNum),
-				Builder.MakeWeightsConstant(ValueNum, 1.0f));
+				Builder.MakeValuesZero(ValueNum),
+				Builder.MakeValuesOne(ValueNum));
 			break;
 		}
 
@@ -965,14 +965,13 @@ namespace UE::Learning::Observation
 				NNE::RuntimeBasic::FModelBuilderElement BuilderSubElement;
 				MakeEncoderNetworkModelBuilderElementFromSchema(BuilderSubElement, Builder, Schema, SubElement, NetworkSettings);
 				BuilderSubLayers.Emplace(BuilderSubElement);
-				if (NetworkSettings.bUseCompressedLinearLayers)
-				{
-					BuilderEncoders.Emplace(Builder.MakeCompressedLinearWithRandomKaimingWeights(SubElementEncodedSize, Parameters.EncodingSize));
-				}
-				else
-				{
-					BuilderEncoders.Emplace(Builder.MakeLinearWithRandomKaimingWeights(SubElementEncodedSize, Parameters.EncodingSize));
-				}
+
+				NNE::RuntimeBasic::FModelBuilder::FLinearLayerSettings LinearLayerSettings;
+				LinearLayerSettings.Type = NetworkSettings.bUseCompressedLinearLayers ?
+					NNE::RuntimeBasic::FModelBuilder::ELinearLayerType::Compressed :
+					NNE::RuntimeBasic::FModelBuilder::ELinearLayerType::Normal;
+
+				BuilderEncoders.Emplace(Builder.MakeLinearLayer(SubElementEncodedSize, Parameters.EncodingSize, LinearLayerSettings));
 			}
 
 			OutElement = Builder.MakeAggregateOrExclusive(Parameters.EncodingSize, BuilderSubLayers, BuilderEncoders);
@@ -996,18 +995,15 @@ namespace UE::Learning::Observation
 				NNE::RuntimeBasic::FModelBuilderElement BuilderSubElement;
 				MakeEncoderNetworkModelBuilderElementFromSchema(BuilderSubElement, Builder, Schema, SubElement, NetworkSettings);
 				BuilderSubLayers.Emplace(BuilderSubElement);
-				if (NetworkSettings.bUseCompressedLinearLayers)
-				{
-					BuilderQueryLayers.Emplace(Builder.MakeCompressedLinearWithRandomKaimingWeights(SubElementEncodedSize, Parameters.AttentionHeadNum * Parameters.AttentionEncodingSize));
-					BuilderKeyLayers.Emplace(Builder.MakeCompressedLinearWithRandomKaimingWeights(SubElementEncodedSize, Parameters.AttentionHeadNum * Parameters.AttentionEncodingSize));
-					BuilderValueLayers.Emplace(Builder.MakeCompressedLinearWithRandomKaimingWeights(SubElementEncodedSize, Parameters.AttentionHeadNum * Parameters.ValueEncodingSize));
-				}
-				else
-				{
-					BuilderQueryLayers.Emplace(Builder.MakeLinearWithRandomKaimingWeights(SubElementEncodedSize, Parameters.AttentionHeadNum * Parameters.AttentionEncodingSize));
-					BuilderKeyLayers.Emplace(Builder.MakeLinearWithRandomKaimingWeights(SubElementEncodedSize, Parameters.AttentionHeadNum * Parameters.AttentionEncodingSize));
-					BuilderValueLayers.Emplace(Builder.MakeLinearWithRandomKaimingWeights(SubElementEncodedSize, Parameters.AttentionHeadNum * Parameters.ValueEncodingSize));
-				}
+
+				NNE::RuntimeBasic::FModelBuilder::FLinearLayerSettings LinearLayerSettings;
+				LinearLayerSettings.Type = NetworkSettings.bUseCompressedLinearLayers ?
+					NNE::RuntimeBasic::FModelBuilder::ELinearLayerType::Compressed :
+					NNE::RuntimeBasic::FModelBuilder::ELinearLayerType::Normal;
+
+				BuilderQueryLayers.Emplace(Builder.MakeLinearLayer(SubElementEncodedSize, Parameters.AttentionHeadNum * Parameters.AttentionEncodingSize, LinearLayerSettings));
+				BuilderKeyLayers.Emplace(Builder.MakeLinearLayer(SubElementEncodedSize, Parameters.AttentionHeadNum * Parameters.AttentionEncodingSize, LinearLayerSettings));
+				BuilderValueLayers.Emplace(Builder.MakeLinearLayer(SubElementEncodedSize, Parameters.AttentionHeadNum * Parameters.ValueEncodingSize, LinearLayerSettings));
 			}
 
 			OutElement = Builder.MakeAggregateOrInclusive(
@@ -1041,30 +1037,20 @@ namespace UE::Learning::Observation
 			NNE::RuntimeBasic::FModelBuilderElement BuilderSubElement;
 			MakeEncoderNetworkModelBuilderElementFromSchema(BuilderSubElement, Builder, Schema, Parameters.Element, NetworkSettings);
 
-			if (NetworkSettings.bUseCompressedLinearLayers)
-			{
-				OutElement = Builder.MakeAggregateSet(
-					Parameters.MaxNum,
-					Parameters.ValueEncodingSize,
-					Parameters.AttentionEncodingSize,
-					Parameters.AttentionHeadNum,
-					BuilderSubElement,
-					Builder.MakeCompressedLinearWithRandomKaimingWeights(SubElementEncodedSize, Parameters.AttentionHeadNum * Parameters.AttentionEncodingSize),
-					Builder.MakeCompressedLinearWithRandomKaimingWeights(SubElementEncodedSize, Parameters.AttentionHeadNum * Parameters.AttentionEncodingSize),
-					Builder.MakeCompressedLinearWithRandomKaimingWeights(SubElementEncodedSize, Parameters.AttentionHeadNum * Parameters.ValueEncodingSize));
-			}
-			else
-			{
-				OutElement = Builder.MakeAggregateSet(
-					Parameters.MaxNum,
-					Parameters.ValueEncodingSize,
-					Parameters.AttentionEncodingSize,
-					Parameters.AttentionHeadNum,
-					BuilderSubElement,
-					Builder.MakeLinearWithRandomKaimingWeights(SubElementEncodedSize, Parameters.AttentionHeadNum * Parameters.AttentionEncodingSize),
-					Builder.MakeLinearWithRandomKaimingWeights(SubElementEncodedSize, Parameters.AttentionHeadNum * Parameters.AttentionEncodingSize),
-					Builder.MakeLinearWithRandomKaimingWeights(SubElementEncodedSize, Parameters.AttentionHeadNum * Parameters.ValueEncodingSize));
-			}
+			NNE::RuntimeBasic::FModelBuilder::FLinearLayerSettings LinearLayerSettings;
+			LinearLayerSettings.Type = NetworkSettings.bUseCompressedLinearLayers ?
+				NNE::RuntimeBasic::FModelBuilder::ELinearLayerType::Compressed :
+				NNE::RuntimeBasic::FModelBuilder::ELinearLayerType::Normal;
+
+			OutElement = Builder.MakeAggregateSet(
+				Parameters.MaxNum,
+				Parameters.ValueEncodingSize,
+				Parameters.AttentionEncodingSize,
+				Parameters.AttentionHeadNum,
+				BuilderSubElement,
+				Builder.MakeLinearLayer(SubElementEncodedSize, Parameters.AttentionHeadNum * Parameters.AttentionEncodingSize, LinearLayerSettings),
+				Builder.MakeLinearLayer(SubElementEncodedSize, Parameters.AttentionHeadNum * Parameters.AttentionEncodingSize, LinearLayerSettings),
+				Builder.MakeLinearLayer(SubElementEncodedSize, Parameters.AttentionHeadNum * Parameters.ValueEncodingSize, LinearLayerSettings));
 
 			break;
 		}
@@ -1078,32 +1064,22 @@ namespace UE::Learning::Observation
 			NNE::RuntimeBasic::FModelBuilderElement BuilderSubElement;
 			MakeEncoderNetworkModelBuilderElementFromSchema(BuilderSubElement, Builder, Schema, Parameters.Element, NetworkSettings);
 
-			if (NetworkSettings.bUseCompressedLinearLayers)
-			{
-				OutElement = Builder.MakeSequence({
-					BuilderSubElement,
-					Builder.MakeCompressedMLPWithRandomKaimingWeights(
-						SubElementEncodedSize,
-						Parameters.EncodingSize,
-						Parameters.EncodingSize,
-						Parameters.LayerNum + 1, // Add 1 to account for input layer
-						Private::GetNNEActivationFunction(Parameters.ActivationFunction),
-						true)
-					});
-			}
-			else
-			{
-				OutElement = Builder.MakeSequence({
-					BuilderSubElement,
-					Builder.MakeMLPWithRandomKaimingWeights(
-						SubElementEncodedSize,
-						Parameters.EncodingSize,
-						Parameters.EncodingSize,
-						Parameters.LayerNum + 1, // Add 1 to account for input layer
-						Private::GetNNEActivationFunction(Parameters.ActivationFunction),
-						true)
-					});
-			}
+			NNE::RuntimeBasic::FModelBuilder::FLinearLayerSettings LinearLayerSettings;
+			LinearLayerSettings.Type = NetworkSettings.bUseCompressedLinearLayers ?
+				NNE::RuntimeBasic::FModelBuilder::ELinearLayerType::Compressed :
+				NNE::RuntimeBasic::FModelBuilder::ELinearLayerType::Normal;
+
+			OutElement = Builder.MakeSequence({
+				BuilderSubElement,
+				Builder.MakeMLP(
+					SubElementEncodedSize,
+					Parameters.EncodingSize,
+					Parameters.EncodingSize,
+					Parameters.LayerNum + 1, // Add 1 to account for input layer
+					Private::GetNNEActivationFunction(Parameters.ActivationFunction),
+					true,
+					LinearLayerSettings)
+				});
 
 			break;
 		}
diff --git a/Engine/Plugins/Experimental/LearningAgents/Source/LearningAgents/Private/LearningAgentsCritic.cpp b/Engine/Plugins/Experimental/LearningAgents/Source/LearningAgents/Private/LearningAgentsCritic.cpp
--- a/Engine/Plugins/Experimental/LearningAgents/Source/LearningAgents/Private/LearningAgentsCritic.cpp
+++ b/Engine/Plugins/Experimental/LearningAgents/Source/LearningAgents/Private/LearningAgentsCritic.cpp
@@ -179,7 +179,7 @@ void ULearningAgentsCritic::SetupCritic(
 		TArray<uint8> FileData;
 		uint32 CriticInputSize, CriticOutputSize;
 		Builder.WriteFileDataAndReset(FileData, CriticInputSize, CriticOutputSize,
-			Builder.MakeMLPWithRandomKaimingWeights(
+			Builder.MakeMLP(
 				ObservationEncodedVectorSize + MemoryStateSize,
 				1,
 				CriticSettings.HiddenLayerSize,
diff --git a/Engine/Plugins/Experimental/LearningAgents/Source/LearningAgents/Private/LearningAgentsPolicy.cpp b/Engine/Plugins/Experimental/LearningAgents/Source/LearningAgents/Private/LearningAgentsPolicy.cpp
--- a/Engine/Plugins/Experimental/LearningAgents/Source/LearningAgents/Private/LearningAgentsPolicy.cpp
+++ b/Engine/Plugins/Experimental/LearningAgents/Source/LearningAgents/Private/LearningAgentsPolicy.cpp
@@ -260,25 +260,29 @@ void ULearningAgentsPolicy::SetupPolicy(
 		const ELearningAgentsActivationFunction PolicyActivationFunction = PolicySettings.ActivationFunction;
 		const float PolicyInitialEncodedActionScale = PolicySettings.InitialEncodedActionScale;
 
+		// Scale down memory cell initial weights to improve stability at beginning of training
+		UE::NNE::RuntimeBasic::FModelBuilder::FLinearLayerSettings MemoryCellWeightSettings;
+		MemoryCellWeightSettings.WeightInitializationSettings.Scale = 0.1f;
+
 		TArray<uint8> FileData;
 		uint32 PolicyInputSize, PolicyOutputSize;
 		Builder.WriteFileDataAndReset(FileData, PolicyInputSize, PolicyOutputSize,
 			Builder.MakeMemoryBackbone(
 				MemoryStateSize,
-				Builder.MakeMLPWithRandomKaimingWeights(
+				Builder.MakeMLP(
 					ObservationEncodedVectorSize,
 					PolicyHiddenLayerSize,
 					PolicyHiddenLayerSize,
 					PolicyHiddenLayerNum / 2 + 2, // Add 2 to account for input and output layers
 					UE::Learning::Agents::Policy::Private::GetBuilderActivationFunction(PolicyActivationFunction),
 					true),
-				Builder.MakeMemoryCellWithLinearRandomKaimingWeights(
+				Builder.MakeMemoryCellLayer(
 					PolicyHiddenLayerSize,
 					PolicyHiddenLayerSize,
 					MemoryStateSize,
-					0.1f),
+					MemoryCellWeightSettings),
 				Builder.MakeSequence({
-					Builder.MakeMLPWithRandomKaimingWeights(
+					Builder.MakeMLP(
 						PolicyHiddenLayerSize,
 						ActionEncodedVectorSize,
 						PolicyHiddenLayerSize,
@@ -286,8 +290,8 @@ void ULearningAgentsPolicy::SetupPolicy(
 						UE::Learning::Agents::Policy::Private::GetBuilderActivationFunction(PolicyActivationFunction)),
 					Builder.MakeDenormalize(
 						ActionEncodedVectorSize,
-						Builder.MakeWeightsZero(ActionEncodedVectorSize),
-						Builder.MakeWeightsConstant(ActionEncodedVectorSize, PolicyInitialEncodedActionScale))
+						Builder.MakeValuesZero(ActionEncodedVectorSize),
+						Builder.MakeValuesConstant(ActionEncodedVectorSize, PolicyInitialEncodedActionScale))
 					})
 			));
 
diff --git a/Engine/Plugins/Experimental/NNERuntimeBasicCpu/Source/NNERuntimeBasicCpu/Private/NNERuntimeBasicCpuModel.cpp b/Engine/Plugins/Experimental/NNERuntimeBasicCpu/Source/NNERuntimeBasicCpu/Private/NNERuntimeBasicCpuModel.cpp
--- a/Engine/Plugins/Experimental/NNERuntimeBasicCpu/Source/NNERuntimeBasicCpu/Private/NNERuntimeBasicCpuModel.cpp
+++ b/Engine/Plugins/Experimental/NNERuntimeBasicCpu/Source/NNERuntimeBasicCpu/Private/NNERuntimeBasicCpuModel.cpp
@@ -4341,8 +4341,17 @@ namespace UE::NNE::RuntimeBasic
 		{
 			State = RngInt(State ^ 0x0c32dd74);
 		}
+
+		static inline float RngClippedGaussian(const uint32 State, const float Clip = 10.0f)
+		{
+			return FMath::Clamp(RngGaussian(State), -Clip, Clip);
+		}
+
 	}
 
+	FModelBuilder::FWeightInitializationSettings::FWeightInitializationSettings() = default;
+	FModelBuilder::FLinearLayerSettings::FLinearLayerSettings() = default;
+
 	FModelBuilder::FModelBuilder(int32 Seed) : Rng(Seed) {};
 
 	FModelBuilderElement::FModelBuilderElement() = default;
@@ -4401,35 +4410,79 @@ namespace UE::NNE::RuntimeBasic
 		return StaticCastSharedPtr<Private::ILayer>(LinearLayer);
 	}
 
-	FModelBuilderElement FModelBuilder::MakeLinearWithRandomKaimingWeights(
+	FModelBuilderElement FModelBuilder::MakeLipschiztLinear(
 		const uint32 InputSize,
 		const uint32 OutputSize,
-		const float WeightScale)
+		const TConstArrayView<float> Weights,
+		const TConstArrayView<float> Biases)
 	{
-		return MakeLinear(
-			InputSize,
-			OutputSize,
-			MakeWeightsRandomKaiming(InputSize, OutputSize, WeightScale),
-			MakeWeightsZero(OutputSize));
+		check(Biases.Num() == OutputSize);
+		check(Weights.Num() == InputSize * OutputSize);
+
+		const TSharedPtr<Private::FLipschiztLinearLayer> LipschiztLinearLayer = MakeShared<Private::FLipschiztLinearLayer>();
+		LipschiztLinearLayer->InputSize = InputSize;
+		LipschiztLinearLayer->OutputSize = OutputSize;
+		LipschiztLinearLayer->Biases = Biases;
+		LipschiztLinearLayer->Weights = Weights;
+
+		return StaticCastSharedPtr<Private::ILayer>(LipschiztLinearLayer);
 	}
 
-	FModelBuilderElement FModelBuilder::MakeCompressedLinearWithRandomKaimingWeights(
+	FModelBuilderElement FModelBuilder::MakeLinearLayer(
 		const uint32 InputSize,
 		const uint32 OutputSize,
-		const float WeightScale)
+		const FLinearLayerSettings& Settings)
 	{
-		TArrayView<uint16> Weights;
-		TArrayView<float> WeightOffsets;
-		TArrayView<float> WeightScales;
-		MakeCompressedWeightsRandomKaiming(Weights, WeightOffsets, WeightScales, InputSize, OutputSize, WeightScale);
+		switch (Settings.Type)
+		{
+		case ELinearLayerType::Normal:
+		{
+			return MakeLinear(
+				InputSize,
+				OutputSize,
+				MakeInitialWeights(InputSize, OutputSize, Settings.WeightInitializationSettings),
+				MakeInitialBiases(OutputSize, Settings.WeightInitializationSettings));
+		}
+					
+		case ELinearLayerType::Compressed:
+		{
+			TArrayView<uint16> Weights;
+			TArrayView<float> WeightOffsets;
+			TArrayView<float> WeightScales;
 
-		return MakeCompressedLinear(
-			InputSize,
-			OutputSize,
-			Weights,
-			WeightOffsets,
-			WeightScales,
-			MakeWeightsZero(OutputSize));
+			MakeInitialCompressedWeights(
+				Weights,
+				WeightOffsets,
+				WeightScales,
+				InputSize,
+				OutputSize,
+				Settings.WeightInitializationSettings);
+
+			return MakeCompressedLinear(
+				InputSize,
+				OutputSize,
+				Weights,
+				WeightOffsets,
+				WeightScales,
+				MakeInitialBiases(OutputSize, Settings.WeightInitializationSettings));
+		}
+
+		case ELinearLayerType::Lipschizt:
+		{
+			return MakeLipschiztLinear(
+				InputSize,
+				OutputSize,
+				MakeInitialWeights(InputSize, OutputSize, Settings.WeightInitializationSettings),
+				MakeInitialBiases(OutputSize, Settings.WeightInitializationSettings));
+		}
+
+		default:
+		{
+			checkNoEntry();
+			return FModelBuilderElement();
+		}
+
+		}
 	}
 
 	FModelBuilderElement FModelBuilder::MakeMultiLinear(
@@ -4576,25 +4629,29 @@ namespace UE::NNE::RuntimeBasic
 		return StaticCastSharedPtr<Private::ILayer>(SequenceLayer);
 	}
 
-	FModelBuilderElement FModelBuilder::MakeMLPWithRandomKaimingWeights(
+	FModelBuilderElement FModelBuilder::MakeMLP(
 		const uint32 InputSize,
 		const uint32 OutputSize,
 		const uint32 HiddenSize,
 		const uint32 LayerNum,
 		const EActivationFunction ActivationFunction,
-		const bool bActivationOnFinalLayer)
+		const bool bActivationOnFinalLayer,
+		const FLinearLayerSettings& LinearLayerSettings)
 	{
 		check(LayerNum >= 2);
 
+		// Reserve space for Linear + Activation on each layer, minus the final activation if it is not required.
+		const int32 TotalLayerNum = 2 * LayerNum - (bActivationOnFinalLayer ? 0 : 1);
+
 		TArray<FModelBuilderElement, TInlineAllocator<32>> Layers;
-		Layers.Reserve(2 * LayerNum - (bActivationOnFinalLayer ? 0 : 1));
+		Layers.Reserve(TotalLayerNum);
 
 		for (uint32 LayerIdx = 0; LayerIdx < LayerNum; LayerIdx++)
 		{
 			const uint32 LayerInputSize = LayerIdx == 0 ? InputSize : HiddenSize;
 			const uint32 LayerOuputSize = LayerIdx == LayerNum - 1 ? OutputSize : HiddenSize;
 
-			Layers.Emplace(MakeLinearWithRandomKaimingWeights(LayerInputSize, LayerOuputSize));
+			Layers.Emplace(MakeLinearLayer(LayerInputSize, LayerOuputSize, LinearLayerSettings));
 			
 			if (bActivationOnFinalLayer || LayerIdx != LayerNum - 1)
 			{
@@ -4602,34 +4659,42 @@ namespace UE::NNE::RuntimeBasic
 			}
 		}
 		
+		check(Layers.Num() == TotalLayerNum);
+
 		return MakeSequence(Layers);
 	}
 
-	FModelBuilderElement FModelBuilder::MakeMLPWithRandomKaimingWeightsWithLayerNorm(
+	FModelBuilderElement FModelBuilder::MakeMLPWithLayerNorm(
 		const uint32 InputSize, 
 		const uint32 OutputSize, 
 		const uint32 HiddenSize, 
 		const uint32 LayerNum, 
 		const EActivationFunction ActivationFunction, 
-		const bool bActivationOnFinalLayer)
+		const bool bActivationOnFinalLayer,
+		const FLinearLayerSettings& LinearLayerSettings)
 	{
 		check(LayerNum >= 2);
 
+		// Reserve space for Linear + Activation on each layer, minus the final activation if it is not required, 
+		// plus the LayerNorm in between each layer.
+		const int32 TotalLayerNum = 2 * LayerNum - (bActivationOnFinalLayer ? 0 : 1) + LayerNum - 1;
+
 		TArray<FModelBuilderElement, TInlineAllocator<32>> Layers;
-		Layers.Reserve(2 * LayerNum - (bActivationOnFinalLayer ? 0 : 1) + LayerNum - 1);
+		Layers.Reserve(TotalLayerNum);
 
 		for (uint32 LayerIdx = 0; LayerIdx < LayerNum; LayerIdx++)
 		{
 			const uint32 LayerInputSize = LayerIdx == 0 ? InputSize : HiddenSize;
 			const uint32 LayerOuputSize = LayerIdx == LayerNum - 1 ? OutputSize : HiddenSize;
 
-			Layers.Emplace(MakeLinearWithRandomKaimingWeights(LayerInputSize, LayerOuputSize));
+			Layers.Emplace(MakeLinearLayer(LayerInputSize, LayerOuputSize, LinearLayerSettings));
 			
 			if (LayerIdx != LayerNum - 1)
 			{
-				Layers.Emplace(
-					MakeLayerNorm(LayerOuputSize, MakeWeightsZero(LayerOuputSize), MakeWeightsConstant(LayerOuputSize, 1.0f))
-				);
+				Layers.Emplace(MakeLayerNorm(
+					LayerOuputSize, 
+					MakeValuesZero(LayerOuputSize), 
+					MakeValuesOne(LayerOuputSize)));
 			}
 
 			if (bActivationOnFinalLayer || LayerIdx != LayerNum - 1)
@@ -4638,34 +4703,7 @@ namespace UE::NNE::RuntimeBasic
 			}
 		}
 
-		return MakeSequence(Layers);
-	}
-
-	FModelBuilderElement FModelBuilder::MakeCompressedMLPWithRandomKaimingWeights(
-		const uint32 InputSize,
-		const uint32 OutputSize,
-		const uint32 HiddenSize,
-		const uint32 LayerNum,
-		const EActivationFunction ActivationFunction,
-		const bool bActivationOnFinalLayer)
-	{
-		check(LayerNum >= 2);
-
-		TArray<FModelBuilderElement, TInlineAllocator<32>> Layers;
-		Layers.Reserve(2 * LayerNum - (bActivationOnFinalLayer ? 0 : 1));
-
-		for (uint32 LayerIdx = 0; LayerIdx < LayerNum; LayerIdx++)
-		{
-			const uint32 LayerInputSize = LayerIdx == 0 ? InputSize : HiddenSize;
-			const uint32 LayerOuputSize = LayerIdx == LayerNum - 1 ? OutputSize : HiddenSize;
-
-			Layers.Emplace(MakeCompressedLinearWithRandomKaimingWeights(LayerInputSize, LayerOuputSize));
-
-			if (bActivationOnFinalLayer || LayerIdx != LayerNum - 1)
-			{
-				Layers.Emplace(MakeActivation(LayerOuputSize, ActivationFunction));
-			}
-		}
+		check(Layers.Num() == TotalLayerNum);
 
 		return MakeSequence(Layers);
 	}
@@ -4704,38 +4742,21 @@ namespace UE::NNE::RuntimeBasic
 		return StaticCastSharedPtr<Private::ILayer>(CellLayer);
 	}
 
-	FModelBuilderElement FModelBuilder::MakeMemoryCellWithLinearRandomKaimingWeights(
+	FModelBuilderElement FModelBuilder::MakeMemoryCellLayer(
 		const uint32 InputNum,
 		const uint32 OutputNum,
 		const uint32 MemoryNum,
-		const float WeightScale)
+		const FLinearLayerSettings& LinearLayerSettings)
 	{
 		return MakeMemoryCell(
 			InputNum,
 			OutputNum,
 			MemoryNum,
-			MakeLinearWithRandomKaimingWeights(InputNum + MemoryNum, MemoryNum, WeightScale),
-			MakeLinearWithRandomKaimingWeights(InputNum + MemoryNum, OutputNum, WeightScale),
-			MakeLinearWithRandomKaimingWeights(InputNum + MemoryNum, MemoryNum, WeightScale),
-			MakeLinearWithRandomKaimingWeights(InputNum + MemoryNum, OutputNum, WeightScale),
-			MakeLinearWithRandomKaimingWeights(MemoryNum, OutputNum, WeightScale));
-	}
-
-	FModelBuilderElement FModelBuilder::MakeMemoryCellWithCompressedLinearRandomKaimingWeights(
-		const uint32 InputNum,
-		const uint32 OutputNum,
-		const uint32 MemoryNum,
-		const float WeightScale)
-	{
-		return MakeMemoryCell(
-			InputNum,
-			OutputNum,
-			MemoryNum,
-			MakeCompressedLinearWithRandomKaimingWeights(InputNum + MemoryNum, MemoryNum, WeightScale),
-			MakeCompressedLinearWithRandomKaimingWeights(InputNum + MemoryNum, OutputNum, WeightScale),
-			MakeCompressedLinearWithRandomKaimingWeights(InputNum + MemoryNum, MemoryNum, WeightScale),
-			MakeCompressedLinearWithRandomKaimingWeights(InputNum + MemoryNum, OutputNum, WeightScale),
-			MakeCompressedLinearWithRandomKaimingWeights(MemoryNum, OutputNum, WeightScale));
+			MakeLinearLayer(InputNum + MemoryNum, MemoryNum, LinearLayerSettings),
+			MakeLinearLayer(InputNum + MemoryNum, OutputNum, LinearLayerSettings),
+			MakeLinearLayer(InputNum + MemoryNum, MemoryNum, LinearLayerSettings),
+			MakeLinearLayer(InputNum + MemoryNum, OutputNum, LinearLayerSettings),
+			MakeLinearLayer(MemoryNum, OutputNum, LinearLayerSettings));
 	}
 
 	FModelBuilderElement FModelBuilder::MakeMemoryBackbone(
@@ -4950,24 +4971,6 @@ namespace UE::NNE::RuntimeBasic
 		return StaticCastSharedPtr<Private::ILayer>(LayerNormLayer);
 	}
 
-	FModelBuilderElement FModelBuilder::MakeLipschiztLinear(
-		const uint32 InputSize,
-		const uint32 OutputSize,
-		const TConstArrayView<float> Weights,
-		const TConstArrayView<float> Biases)
-	{
-		check(Biases.Num() == OutputSize);
-		check(Weights.Num() == InputSize * OutputSize);
-
-		const TSharedPtr<Private::FLipschiztLinearLayer> LipschiztLinearLayer = MakeShared<Private::FLipschiztLinearLayer>();
-		LipschiztLinearLayer->InputSize = InputSize;
-		LipschiztLinearLayer->OutputSize = OutputSize;
-		LipschiztLinearLayer->Biases = Biases;
-		LipschiztLinearLayer->Weights = Weights;
-
-		return StaticCastSharedPtr<Private::ILayer>(LipschiztLinearLayer);
-	}
-
 	void FModelBuilder::Reset()
 	{
 		Rng = RngInitialState;
@@ -5022,28 +5025,33 @@ namespace UE::NNE::RuntimeBasic
 		WriteFileDataAndReset(MakeArrayView(FileData), OutInputSize, OutOutputSize, Element);
 	}
 
-	TArrayView<float> FModelBuilder::MakeWeightsCopy(const TConstArrayView<float> Weights)
+	TArrayView<float> FModelBuilder::MakeValuesCopy(const TConstArrayView<float> Values)
 	{
-		TArray<float>& Values = WeightsPool.AddDefaulted_GetRef();
-		Values = Weights;
-		return Values;
+		TArray<float>& OutValues = WeightsPool.AddDefaulted_GetRef();
+		OutValues = Values;
+		return OutValues;
 	}
 
-	TArrayView<float> FModelBuilder::MakeWeightsZero(const uint32 Size)
+	TArrayView<float> FModelBuilder::MakeValuesZero(const uint32 Size)
 	{
 		TArray<float>& Values = WeightsPool.AddDefaulted_GetRef();
 		Values.Init(0.0f, Size);
 		return Values;
 	}
 
-	TArrayView<float> FModelBuilder::MakeWeightsConstant(const uint32 Size, const float Value)
+	TArrayView<float> FModelBuilder::MakeValuesOne(const uint32 Size)
+	{
+		return MakeValuesConstant(Size, 1.0f);
+	}
+
+	TArrayView<float> FModelBuilder::MakeValuesConstant(const uint32 Size, const float Value)
 	{
 		TArray<float>& Values = WeightsPool.AddDefaulted_GetRef();
 		Values.Init(Value, Size);
 		return Values;
 	}
 
-	TArrayView<float> FModelBuilder::MakeWeightsRandomKaiming(const uint32 InputSize, const uint32 OutputSize, const float Scale)
+	TArrayView<float> FModelBuilder::MakeWeightsRandomKaimingGaussian(const uint32 InputSize, const uint32 OutputSize, const float Scale)
 	{
 		TArray<float>& Values = WeightsPool.AddDefaulted_GetRef();
 		Values.SetNumUninitialized(InputSize * OutputSize);
@@ -5052,14 +5060,64 @@ namespace UE::NNE::RuntimeBasic
 
 		for (uint32 Idx = 0; Idx < InputSize * OutputSize; Idx++)
 		{
-			Values[Idx] = Std * Private::RngGaussian(Rng);
+			Values[Idx] = Std * Private::RngClippedGaussian(Rng);
 			Private::RngUpdate(Rng);
 		}
 		
 		return Values;
 	}
 
-	void FModelBuilder::MakeCompressedWeightsRandomKaiming(
+	TArrayView<float> FModelBuilder::MakeWeightsRandomKaimingUniform(const uint32 InputSize, const uint32 OutputSize, const float Scale)
+	{
+		TArray<float>& Values = WeightsPool.AddDefaulted_GetRef();
+		Values.SetNumUninitialized(InputSize * OutputSize);
+
+		const float Gain = FMath::Sqrt(2.0f / (1.0f + 5.0f));
+		const float Std = Gain / FMath::Sqrt((float)InputSize);
+		const float Bound = Scale * FMath::Sqrt(3.0f) * Std;
+
+		for (uint32 Idx = 0; Idx < InputSize * OutputSize; Idx++)
+		{
+			Values[Idx] = Bound * (2.0f * Private::RngUniform(Rng) - 1.0f);
+			Private::RngUpdate(Rng);
+		}
+
+		return Values;
+	}
+
+	TArrayView<float> FModelBuilder::MakeBiasesRandomKaimingGaussian(const uint32 Size, const float Scale)
+	{
+		TArray<float>& Values = WeightsPool.AddDefaulted_GetRef();
+		Values.SetNumUninitialized(Size);
+
+		const float Std = Scale * FMath::Sqrt(2.0f / Size);
+
+		for (uint32 Idx = 0; Idx < Size; Idx++)
+		{
+			Values[Idx] = Std * Private::RngClippedGaussian(Rng);
+			Private::RngUpdate(Rng);
+		}
+
+		return Values;
+	}
+
+	TArrayView<float> FModelBuilder::MakeBiasesRandomKaimingUniform(const uint32 Size, const float Scale)
+	{
+		TArray<float>& Values = WeightsPool.AddDefaulted_GetRef();
+		Values.SetNumUninitialized(Size);
+
+		const float Bound = Scale / FMath::Sqrt((float)Size);
+
+		for (uint32 Idx = 0; Idx < Size; Idx++)
+		{
+			Values[Idx] = Bound * (2.0f * Private::RngUniform(Rng) - 1.0f);
+			Private::RngUpdate(Rng);
+		}
+
+		return Values;
+	}
+
+	void FModelBuilder::MakeCompressedWeightsRandomKaimingGaussian(
 		TArrayView<uint16>& OutWeightsView,
 		TArrayView<float>& OutWeightOffsetsView,
 		TArrayView<float>& OutWeightScalesView,
@@ -5076,7 +5134,7 @@ namespace UE::NNE::RuntimeBasic
 
 		for (uint32 Idx = 0; Idx < InputSize * OutputSize; Idx++)
 		{
-			Values[Idx] = Std * Private::RngGaussian(Rng);
+			Values[Idx] = Std * Private::RngClippedGaussian(Rng);
 			Private::RngUpdate(Rng);
 		}
 
@@ -5135,6 +5193,169 @@ namespace UE::NNE::RuntimeBasic
 		OutWeightScalesView = WeightScales;
 	}
 
+	void FModelBuilder::MakeCompressedWeightsRandomKaimingUniform(
+		TArrayView<uint16>& OutWeightsView,
+		TArrayView<float>& OutWeightOffsetsView,
+		TArrayView<float>& OutWeightScalesView,
+		const uint32 InputSize,
+		const uint32 OutputSize,
+		const float Scale)
+	{
+		// Make Kaiming Weights
+
+		TArray<float> Values;
+		Values.SetNumUninitialized(InputSize * OutputSize);
+
+		const float Gain = FMath::Sqrt(2.0f / (1.0f + 5.0f));
+		const float Std = Gain / FMath::Sqrt((float)InputSize);
+		const float Bound = Scale * FMath::Sqrt(3.0f) * Std;
+
+		for (uint32 Idx = 0; Idx < InputSize * OutputSize; Idx++)
+		{
+			Values[Idx] = Bound * (2.0f * Private::RngUniform(Rng) - 1.0f);
+			Private::RngUpdate(Rng);
+		}
+
+		// Find Min and Max
+
+		TArray<float> Mins;
+		TArray<float> Maxs;
+		Mins.SetNumUninitialized(InputSize);
+		Maxs.SetNumUninitialized(InputSize);
+
+		for (uint32 RowIdx = 0; RowIdx < InputSize; RowIdx++)
+		{
+			Mins[RowIdx] = +FLT_MAX;
+			Maxs[RowIdx] = -FLT_MAX;
+		}
+
+		for (uint32 RowIdx = 0; RowIdx < InputSize; RowIdx++)
+		{
+			for (uint32 ColIdx = 0; ColIdx < OutputSize; ColIdx++)
+			{
+				Mins[RowIdx] = FMath::Min(Mins[RowIdx], Values[RowIdx * OutputSize + ColIdx]);
+				Maxs[RowIdx] = FMath::Max(Maxs[RowIdx], Values[RowIdx * OutputSize + ColIdx]);
+			}
+		}
+
+		// Find Scale and Offset
+
+		TArray<float>& WeightOffsets = WeightsPool.AddDefaulted_GetRef();
+		WeightOffsets.SetNumUninitialized(InputSize);
+
+		TArray<float>& WeightScales = WeightsPool.AddDefaulted_GetRef();
+		WeightScales.SetNumUninitialized(InputSize);
+
+		for (uint32 RowIdx = 0; RowIdx < InputSize; RowIdx++)
+		{
+			WeightOffsets[RowIdx] = Mins[RowIdx];
+			WeightScales[RowIdx] = FMath::Max(Maxs[RowIdx] - Mins[RowIdx], UE_SMALL_NUMBER) / 65535.0;
+		}
+
+		// Compress
+
+		TArray<uint16>& Weights = CompressedWeightsPool.AddDefaulted_GetRef();
+		Weights.SetNumUninitialized(InputSize * OutputSize);
+
+		for (uint32 RowIdx = 0; RowIdx < InputSize; RowIdx++)
+		{
+			for (uint32 ColIdx = 0; ColIdx < OutputSize; ColIdx++)
+			{
+				Weights[RowIdx * OutputSize + ColIdx] = (uint16)FMath::RoundToFloat(65535.0 *
+					FMath::Clamp((Values[RowIdx * OutputSize + ColIdx] - Mins[RowIdx]) / (Maxs[RowIdx] - Mins[RowIdx]), 0.0f, 1.0f));
+			}
+		}
+
+		OutWeightsView = Weights;
+		OutWeightOffsetsView = WeightOffsets;
+		OutWeightScalesView = WeightScales;
+	}
+
+	TArrayView<float> FModelBuilder::MakeInitialWeights(
+		const uint32 InputSize,
+		const uint32 OutputSize,
+		const FWeightInitializationSettings& WeightInitializationSettings)
+	{
+		switch (WeightInitializationSettings.Type)
+		{
+		case EWeightInitializationType::KaimingGaussian:
+
+			return MakeWeightsRandomKaimingGaussian(InputSize, OutputSize, WeightInitializationSettings.Scale);
+
+		case EWeightInitializationType::KaimingUniform:
+
+			return MakeWeightsRandomKaimingUniform(InputSize, OutputSize, WeightInitializationSettings.Scale);
+
+		default:
+			checkNoEntry();
+			return TArrayView<float>();
+		}
+	}
+
+	TArrayView<float> FModelBuilder::MakeInitialBiases(
+		const uint32 OutputSize,
+		const FWeightInitializationSettings& WeightInitializationSettings)
+	{
+		if (!WeightInitializationSettings.bInitializeBiases)
+		{
+			return MakeValuesZero(OutputSize);
+		}
+
+		switch (WeightInitializationSettings.Type)
+		{
+		case EWeightInitializationType::KaimingGaussian:
+
+			return MakeBiasesRandomKaimingGaussian(OutputSize, WeightInitializationSettings.Scale);
+
+		case EWeightInitializationType::KaimingUniform:
+
+			return MakeBiasesRandomKaimingUniform(OutputSize, WeightInitializationSettings.Scale);
+
+		default:
+			checkNoEntry();
+			return TArrayView<float>();
+		}
+	}
+
+	void FModelBuilder::MakeInitialCompressedWeights(
+		TArrayView<uint16>& OutWeightsView,
+		TArrayView<float>& OutWeightOffsetsView,
+		TArrayView<float>& OutWeightScalesView,
+		const uint32 InputSize,
+		const uint32 OutputSize,
+		const FWeightInitializationSettings& WeightInitializationSettings)
+	{
+		switch (WeightInitializationSettings.Type)
+		{
+		case EWeightInitializationType::KaimingGaussian:
+
+			MakeCompressedWeightsRandomKaimingGaussian(
+				OutWeightsView, 
+				OutWeightOffsetsView, 
+				OutWeightScalesView, 
+				InputSize, 
+				OutputSize, 
+				WeightInitializationSettings.Scale);
+
+			return;
+
+		case EWeightInitializationType::KaimingUniform:
+
+			MakeCompressedWeightsRandomKaimingUniform(
+				OutWeightsView,
+				OutWeightOffsetsView,
+				OutWeightScalesView,
+				InputSize,
+				OutputSize,
+				WeightInitializationSettings.Scale);
+
+			return;
+
+		default:
+			checkNoEntry();
+			return;
+		}
+	}
 
 	TArrayView<uint32> FModelBuilder::MakeSizesZero(const uint32 Size)
 	{
diff --git a/Engine/Plugins/Experimental/NNERuntimeBasicCpu/Source/NNERuntimeBasicCpu/Public/NNERuntimeBasicCpuBuilder.h b/Engine/Plugins/Experimental/NNERuntimeBasicCpu/Source/NNERuntimeBasicCpu/Public/NNERuntimeBasicCpuBuilder.h
--- a/Engine/Plugins/Experimental/NNERuntimeBasicCpu/Source/NNERuntimeBasicCpu/Public/NNERuntimeBasicCpuBuilder.h
+++ b/Engine/Plugins/Experimental/NNERuntimeBasicCpu/Source/NNERuntimeBasicCpu/Public/NNERuntimeBasicCpuBuilder.h
@@ -39,7 +39,8 @@ namespace UE::NNE::RuntimeBasic
 	 * out to create a FileData object which can then be loaded with NNE.
 	 * 
 	 * Note: If you pass your own views into this model builder they must out-live the
-	 * builder itself as it will not create internal copies.
+	 * builder itself as it will not create internal copies. If you want to create an
+	 * internal copy use the `MakeValuesCopy` function.
 	 */
 	class NNERUNTIMEBASICCPU_API FModelBuilder
 	{
@@ -49,13 +50,66 @@ namespace UE::NNE::RuntimeBasic
 		/** Common activation function types. */
 		enum class EActivationFunction : uint8
 		{
-
 			ReLU = 0,
 			ELU = 1,
 			TanH = 2,
 			GELU = 3,
 		};
 
+		/** Linear Layer types. */
+		enum class ELinearLayerType : uint8
+		{
+			Normal = 0,
+			Compressed = 1,
+			Lipschizt = 2,
+		};
+
+		/** Weight initialization types. */
+		enum class EWeightInitializationType : uint8
+		{
+			// Good all-round default weight initialization
+			KaimingGaussian = 0,
+
+			// Default initialization used in PyTorch for Linear Layers
+			KaimingUniform = 1,
+		};
+
+		/*
+		 * There is a clang compiler bug that requires the following classes to declare a default constructor in the cpp file as a workaround:
+		 * 
+		 * [REDACTED_PROFILE_URL]/llvm-project/issues/36032
+		 * https://stackoverflow.com/a/73248547
+		 */
+
+		/** Weight initialization settings. */
+		struct NNERUNTIMEBASICCPU_API FWeightInitializationSettings
+		{
+			// Workaround for clang compiler bug...
+			FWeightInitializationSettings();
+
+			// Type of initialization to use
+			EWeightInitializationType Type = EWeightInitializationType::KaimingGaussian;
+
+			// Overall scale of the weight initialization
+			float Scale = 1.0f;
+
+			// If true, initializes biases using the given initialization method, otherwise initializes them to zeros
+			bool bInitializeBiases = false;
+		};
+
+		/** Linear Layer settings including the type of layer and the way the weights are initialized */
+		struct NNERUNTIMEBASICCPU_API FLinearLayerSettings
+		{
+			// Workaround for clang compiler bug...
+			FLinearLayerSettings();
+
+			// Type of the Linear Layer to construct
+			ELinearLayerType Type = ELinearLayerType::Normal;
+
+			// Weight initialization Settings for the linear Layer
+			FWeightInitializationSettings WeightInitializationSettings;
+		};
+
 		/** Construct a new Model Builder with the given random seed. */
 		FModelBuilder(int32 Seed = 0x0a974e75);
 
@@ -92,28 +146,30 @@ namespace UE::NNE::RuntimeBasic
 			const TConstArrayView<float> Biases);
 
 		/**
-		 * Makes a new linear layer with randomly initialized weights using the Kaiming method.
+		 * Makes a new Lipschizt linear layer.
 		 *
 		 * @param InputSize		Input Vector Size (Number of Rows)
 		 * @param OutputSize	Output Vector Size (Number of Columns)
-		 * @param WeightScale	Scaling factor for the weight creation.
+		 * @param Weights		Linear layer weights.
+		 * @param Biases		Linear layer biases.
 		 */
-		FModelBuilderElement MakeLinearWithRandomKaimingWeights(
+		FModelBuilderElement MakeLipschiztLinear(
 			const uint32 InputSize,
 			const uint32 OutputSize,
-			const float WeightScale = 1.0f);
+			const TConstArrayView<float> Weights,
+			const TConstArrayView<float> Biases);
 
 		/**
-		 * Makes a new compressed linear layer with randomly initialized weights using the Kaiming method.
+		 * Makes a new linear layer with the given settings.
 		 *
 		 * @param InputSize		Input Vector Size (Number of Rows)
 		 * @param OutputSize	Output Vector Size (Number of Columns)
-		 * @param WeightScale	Scaling factor for the weight creation.
+		 * @param Settings		Layer Settings.
 		 */
-		FModelBuilderElement MakeCompressedLinearWithRandomKaimingWeights(
+		FModelBuilderElement MakeLinearLayer(
 			const uint32 InputSize,
 			const uint32 OutputSize,
-			const float WeightScale = 1.0f);
+			const FLinearLayerSettings& Settings = FLinearLayerSettings());
 
 		/**
 		 * Makes a new multi linear layer.
@@ -141,6 +197,7 @@ namespace UE::NNE::RuntimeBasic
 			const uint32 InputOutputSize,
 			const TConstArrayView<float> Mean,
 			const TConstArrayView<float> Std);
+
 		/**
 		 * Makes a new Denormalization layer.
 		 *
@@ -187,7 +244,7 @@ namespace UE::NNE::RuntimeBasic
 		FModelBuilderElement MakeSequence(const TConstArrayView<FModelBuilderElement> Elements);
 
 		/**
-		 * Makes a Multi-Layer Perceptron network with randomly initialized weights using the Kaiming method.
+		 * Makes a Multi-Layer Perceptron network.
 		 *
 		 * @param InputSize					Input Vector Size
 		 * @param OutputSize				Output Vector Size
@@ -195,36 +252,19 @@ namespace UE::NNE::RuntimeBasic
 		 * @param LayerNum					Number of layers. Includes input and output layers.
 		 * @param ActivationFunction		Activation function to use.
 		 * @param bActivationOnFinalLayer	If the activation function should be used on the final output layer.
+		 * @param LinearLayerSettings		Settings for the linear layers in the MLP.
 		 */
-		FModelBuilderElement MakeMLPWithRandomKaimingWeights(
+		FModelBuilderElement MakeMLP(
 			const uint32 InputSize,
 			const uint32 OutputSize, 
 			const uint32 HiddenSize,
 			const uint32 LayerNum,
 			const EActivationFunction ActivationFunction,
-			const bool bActivationOnFinalLayer = false);
-
-
-		/**
-		 * Makes a Multi-Layer Perceptron network with randomly initialized weights using the Kaiming method with LayerNorm before all activations.
-		 *
-		 * @param InputSize					Input Vector Size
-		 * @param OutputSize				Output Vector Size
-		 * @param HiddenSize				Number of hidden units to use on internal layers.
-		 * @param LayerNum					Number of layers. Includes input and output layers.
-		 * @param ActivationFunction		Activation function to use.
-		 * @param bActivationOnFinalLayer	If the activation function should be used on the final output layer.
-		 */
-		FModelBuilderElement MakeMLPWithRandomKaimingWeightsWithLayerNorm(
-			const uint32 InputSize,
-			const uint32 OutputSize,
-			const uint32 HiddenSize,
-			const uint32 LayerNum,
-			const EActivationFunction ActivationFunction,
-			const bool bActivationOnFinalLayer = false);
+			const bool bActivationOnFinalLayer = false,
+			const FLinearLayerSettings& LinearLayerSettings = FLinearLayerSettings());
 
 		/**
-		 * Makes a Compressed Multi-Layer Perceptron network with randomly initialized weights using the Kaiming method.
+		 * Makes a Multi-Layer Perceptron network with LayerNorm before all activations.
 		 *
 		 * @param InputSize					Input Vector Size
 		 * @param OutputSize				Output Vector Size
@@ -232,14 +272,16 @@ namespace UE::NNE::RuntimeBasic
 		 * @param LayerNum					Number of layers. Includes input and output layers.
 		 * @param ActivationFunction		Activation function to use.
 		 * @param bActivationOnFinalLayer	If the activation function should be used on the final output layer.
+		 * @param LinearLayerSettings		Settings for the linear layers in the MLP.
 		 */
-		FModelBuilderElement MakeCompressedMLPWithRandomKaimingWeights(
+		FModelBuilderElement MakeMLPWithLayerNorm(
 			const uint32 InputSize,
 			const uint32 OutputSize,
 			const uint32 HiddenSize,
 			const uint32 LayerNum,
 			const EActivationFunction ActivationFunction,
-			const bool bActivationOnFinalLayer = false);
+			const bool bActivationOnFinalLayer = false,
+			const FLinearLayerSettings& LinearLayerSettings = FLinearLayerSettings());
 
 		/**
 		 * Make a new Memory Cell layer.
@@ -264,32 +306,18 @@ namespace UE::NNE::RuntimeBasic
 			const FModelBuilderElement& OutputMemoryUpdateLayer);
 
 		/**
-		 * Make a new Memory Cell layer with randomly initialized weights using the Kaiming method.
+		 * Make a new Memory Cell layer with the given linear layer settings.
 		 *
-		 * @param InputNum			Number of normal inputs to the model
-		 * @param OutputNum			Number of normal outputs from the model
-		 * @param MemoryNum			The size of the memory vector used by the model
-		 * @param WeightScale		Scaling factor for the weight creation.
-		 */
-		FModelBuilderElement MakeMemoryCellWithLinearRandomKaimingWeights(
-			const uint32 InputNum,
-			const uint32 OutputNum,
-			const uint32 MemoryNum,
-			const float WeightScale = 1.0f);
-
-		/**
-		 * Make a new Memory Cell layer with compressed weights randomly initialized using the Kaiming method.
-		 *
-		 * @param InputNum			Number of normal inputs to the model
-		 * @param OutputNum			Number of normal outputs from the model
-		 * @param MemoryNum			The size of the memory vector used by the model
-		 * @param WeightScale		Scaling factor for the weight creation.
+		 * @param InputNum				Number of normal inputs to the model
+		 * @param OutputNum				Number of normal outputs from the model
+		 * @param MemoryNum				The size of the memory vector used by the model
+		 * @param LinearLayerSettings	Settings for the linear layers in the Memory Cell.
 		 */
-		FModelBuilderElement MakeMemoryCellWithCompressedLinearRandomKaimingWeights(
+		FModelBuilderElement MakeMemoryCellLayer(
 			const uint32 InputNum,
 			const uint32 OutputNum,
 			const uint32 MemoryNum,
-			const float WeightScale = 1.0f);
+			const FLinearLayerSettings& LinearLayerSettings = FLinearLayerSettings());
 
 		/**
 		 * Make a new Memory Backbone layer.
@@ -380,36 +408,43 @@ namespace UE::NNE::RuntimeBasic
 			const TConstArrayView<float> Scales,
 			const float Epsilon = 1e-5f);
 
-		/**
-		 * Makes a new Lipschizt linear layer.
-		 *
-		 * @param InputSize		Input Vector Size (Number of Rows)
-		 * @param OutputSize	Output Vector Size (Number of Columns)
-		 * @param Weights		Linear layer weights.
-		 * @param Biases		Linear layer biases.
-		 */
-		FModelBuilderElement MakeLipschiztLinear(
-			const uint32 InputSize,
-			const uint32 OutputSize,
-			const TConstArrayView<float> Weights,
-			const TConstArrayView<float> Biases);
-
 	public:
 
-		/** Creates a array of weights from a copy of the given array view */
-		TArrayView<float> MakeWeightsCopy(const TConstArrayView<float> Weights);
+		/** Creates a array of values from a copy of the given array view */
+		TArrayView<float> MakeValuesCopy(const TConstArrayView<float> Values);
+
+		/** Creates a array of values set to zero of the given size */
+		TArrayView<float> MakeValuesZero(const uint32 Size);
+
+		/** Creates a array of values set to one of the given size */
+		TArrayView<float> MakeValuesOne(const uint32 Size);
+
+		/** Creates a array of values set to the provided constant value of the given size */
+		TArrayView<float> MakeValuesConstant(const uint32 Size, const float Value);
+
+		/** Creates a array of weights randomly initialized using the Gaussian Kaiming method */
+		TArrayView<float> MakeWeightsRandomKaimingGaussian(const uint32 InputSize, const uint32 OutputSize, const float Scale = 1.0f);
+
+		/** Creates a array of weights randomly initialized using the Uniform Kaiming method */
+		TArrayView<float> MakeWeightsRandomKaimingUniform(const uint32 InputSize, const uint32 OutputSize, const float Scale = 1.0f);
 
-		/** Creates a array of weights set to zero of the given size */
-		TArrayView<float> MakeWeightsZero(const uint32 Size);
+		/** Creates a array of biases randomly initialized using the Gaussian Kaiming method */
+		TArrayView<float> MakeBiasesRandomKaimingGaussian(const uint32 Size, const float Scale = 1.0f);
 
-		/** Creates a array of weights set to the provided constant value of the given size */
-		TArrayView<float> MakeWeightsConstant(const uint32 Size, const float Value);
+		/** Creates a array of biases randomly initialized using the Uniform Kaiming method# */
+		TArrayView<float> MakeBiasesRandomKaimingUniform(const uint32 Size, const float Scale = 1.0f);
 
-		/** Creates a array of weights randomly initialized using the Kaiming method */
-		TArrayView<float> MakeWeightsRandomKaiming(const uint32 InputSize, const uint32 OutputSize, const float Scale = 1.0f);
+		/** Creates a array of weights randomly initialized using the Gaussian Kaiming method and compresses them */
+		void MakeCompressedWeightsRandomKaimingGaussian(
+			TArrayView<uint16>& OutWeightsView,
+			TArrayView<float>& OutWeightOffsetsView,
+			TArrayView<float>& OutWeightScalesView,
+			const uint32 InputSize,
+			const uint32 OutputSize,
+			const float Scale = 1.0f);
 
-		/** Creates a array of weights randomly initialized using the Kaiming method */
-		void MakeCompressedWeightsRandomKaiming(
+		/** Creates a array of weights randomly initialized using the Uniform Kaiming method and compresses them */
+		void MakeCompressedWeightsRandomKaimingUniform(
 			TArrayView<uint16>& OutWeightsView,
 			TArrayView<float>& OutWeightOffsetsView,
 			TArrayView<float>& OutWeightScalesView,
@@ -417,6 +452,26 @@ namespace UE::NNE::RuntimeBasic
 			const uint32 OutputSize,
 			const float Scale = 1.0f);
 
+		/** Creates a array of weights initialized using the given settings */
+		TArrayView<float> MakeInitialWeights(
+			const uint32 InputSize, 
+			const uint32 OutputSize,
+			const FWeightInitializationSettings& WeightInitializationSettings = FWeightInitializationSettings());
+
+		/** Creates a array of biases initialized using the given settings */
+		TArrayView<float> MakeInitialBiases(
+			const uint32 OutputSize,
+			const FWeightInitializationSettings& WeightInitializationSettings = FWeightInitializationSettings());
+
+		/** Creates a array of weights randomly initialized using the given settings and compresses them */
+		void MakeInitialCompressedWeights(
+			TArrayView<uint16>& OutWeightsView,
+			TArrayView<float>& OutWeightOffsetsView,
+			TArrayView<float>& OutWeightScalesView,
+			const uint32 InputSize,
+			const uint32 OutputSize,
+			const FWeightInitializationSettings& WeightInitializationSettings = FWeightInitializationSettings());
+
 		/** Creates an array of sizes, initialized to zero. */
 		TArrayView<uint32> MakeSizesZero(const uint32 Size);
 
@@ -480,12 +535,13 @@ namespace UE::NNE::RuntimeBasic
 
 	private:
 
+		/** Default initial state for the random number generator. */
 		static constexpr uint32 RngInitialState = 0xafcc2b45;
 
 		/** Random Number State for generating random weights */
 		uint32 Rng = RngInitialState;
 
-		/** Pool of all weights data used by the `MakeWeights` functions. */
+		/** Pool of all weights data used by the `MakeValues`, `MakeWeights` `MakeBiases` functions. */
 		TArray<TArray<float>> WeightsPool;
 
 		/** Pool of all compressed weights data used by the `MakeCompressedWeights` functions. */
