NNERuntimeBasicCpu: Fixed aliasing checks to allow for the case where both pointers are nullptr.

diff --git a/Engine/Plugins/Experimental/NNERuntimeBasicCpu/Source/NNERuntimeBasicCpu/Private/NNERuntimeBasicCpuModel.cpp b/Engine/Plugins/Experimental/NNERuntimeBasicCpu/Source/NNERuntimeBasicCpu/Private/NNERuntimeBasicCpuModel.cpp
--- a/Engine/Plugins/Experimental/NNERuntimeBasicCpu/Source/NNERuntimeBasicCpu/Private/NNERuntimeBasicCpuModel.cpp
+++ b/Engine/Plugins/Experimental/NNERuntimeBasicCpu/Source/NNERuntimeBasicCpu/Private/NNERuntimeBasicCpuModel.cpp
@@ -15,6 +15,9 @@
 #define NNE_RUNTIME_BASIC_ENABLE_PROFILE (!UE_BUILD_SHIPPING)
 //#define NNE_RUNTIME_BASIC_ENABLE_PROFILE 0
 
+#define NNE_RUNTIME_BASIC_ENABLE_ALIASING_CHECK (!UE_BUILD_SHIPPING)
+//#define NNE_RUNTIME_BASIC_ENABLE_ALIASING_CHECK 0
+
 #if NNE_RUNTIME_BASIC_ENABLE_PROFILE
 #define NNE_RUNTIME_BASIC_TRACE_SCOPE(...) TRACE_CPUPROFILER_EVENT_SCOPE(__VA_ARGS__)
 #else
@@ -25,6 +28,12 @@
 #include "NNERuntimeBasicCpu.ispc.generated.h"
 #endif
 
+#if NNE_RUNTIME_BASIC_ENABLE_ALIASING_CHECK
+#define NNE_RUNTIME_BASIC_CHECK_ALIASING(X, Y) check((!X && !Y) || (X != Y));
+#else
+#define NNE_RUNTIME_BASIC_CHECK_ALIASING(...)
+#endif
+
 namespace UE::NNE::RuntimeBasic
 {
 	namespace Private
@@ -203,7 +212,7 @@ namespace UE::NNE::RuntimeBasic
 		{
 			NNE_RUNTIME_BASIC_TRACE_SCOPE(NNE::RuntimeBasic::Private::OperatorCopy);
 
-			check(Output != Input);
+			NNE_RUNTIME_BASIC_CHECK_ALIASING(Output, Input);
 
 #if NNE_RUNTIME_BASIC_ENABLE_ISPC
 			ispc::NNERuntimeBasicCPUOperatorCopy(
@@ -236,7 +245,7 @@ namespace UE::NNE::RuntimeBasic
 		{
 			NNE_RUNTIME_BASIC_TRACE_SCOPE(NNE::RuntimeBasic::Private::OperatorNormalize);
 
-			check(Output != Input);
+			NNE_RUNTIME_BASIC_CHECK_ALIASING(Output, Input);
 
 #if NNE_RUNTIME_BASIC_ENABLE_ISPC
 			ispc::NNERuntimeBasicCPUOperatorNormalize(
@@ -271,7 +280,7 @@ namespace UE::NNE::RuntimeBasic
 		{
 			NNE_RUNTIME_BASIC_TRACE_SCOPE(NNE::RuntimeBasic::Private::OperatorDenormalize);
 
-			check(Output != Input);
+			NNE_RUNTIME_BASIC_CHECK_ALIASING(Output, Input);
 
 #if NNE_RUNTIME_BASIC_ENABLE_ISPC
 			ispc::NNERuntimeBasicCPUOperatorDenormalize(
@@ -306,7 +315,7 @@ namespace UE::NNE::RuntimeBasic
 		{
 			NNE_RUNTIME_BASIC_TRACE_SCOPE(NNE::RuntimeBasic::Private::OperatorClamp);
 
-			check(Output != Input);
+			NNE_RUNTIME_BASIC_CHECK_ALIASING(Output, Input);
 
 #if NNE_RUNTIME_BASIC_ENABLE_ISPC
 			ispc::NNERuntimeBasicCPUOperatorClamp(
@@ -342,7 +351,7 @@ namespace UE::NNE::RuntimeBasic
 		{
 			NNE_RUNTIME_BASIC_TRACE_SCOPE(NNE::RuntimeBasic::Private::OperatorLinear);
 
-			check(Output != Input);
+			NNE_RUNTIME_BASIC_CHECK_ALIASING(Output, Input);
 
 #if NNE_RUNTIME_BASIC_ENABLE_ISPC
 			ispc::NNERuntimeBasicCPUOperatorLinear(
@@ -397,7 +406,7 @@ namespace UE::NNE::RuntimeBasic
 		{
 			NNE_RUNTIME_BASIC_TRACE_SCOPE(NNE::RuntimeBasic::Private::OperatorCompressedLinear);
 
-			check(Output != Input);
+			NNE_RUNTIME_BASIC_CHECK_ALIASING(Output, Input);
 
 #if NNE_RUNTIME_BASIC_ENABLE_ISPC
 			ispc::NNERuntimeBasicCPUOperatorCompressedLinear(
@@ -456,7 +465,7 @@ namespace UE::NNE::RuntimeBasic
 		{
 			NNE_RUNTIME_BASIC_TRACE_SCOPE(NNE::RuntimeBasic::Private::OperatorMultiLinear);
 
-			check(Output != Input);
+			NNE_RUNTIME_BASIC_CHECK_ALIASING(Output, Input);
 
 			// For this function ispc generates slightly less efficient code than the naive C++ implementation so we
 			// don't bother calling out to the ispc version even if it is available
@@ -502,7 +511,7 @@ namespace UE::NNE::RuntimeBasic
 		{
 			NNE_RUNTIME_BASIC_TRACE_SCOPE(NNE::RuntimeBasic::Private::OperatorReLU);
 
-			check(Output != Input);
+			NNE_RUNTIME_BASIC_CHECK_ALIASING(Output, Input);
 
 #if NNE_RUNTIME_BASIC_ENABLE_ISPC
 			ispc::NNERuntimeBasicCPUOperatorReLU(
@@ -533,7 +542,7 @@ namespace UE::NNE::RuntimeBasic
 		{
 			NNE_RUNTIME_BASIC_TRACE_SCOPE(NNE::RuntimeBasic::Private::OperatorELU);
 
-			check(Output != Input);
+			NNE_RUNTIME_BASIC_CHECK_ALIASING(Output, Input);
 
 #if NNE_RUNTIME_BASIC_ENABLE_ISPC
 			ispc::NNERuntimeBasicCPUOperatorELU(
@@ -570,7 +579,7 @@ namespace UE::NNE::RuntimeBasic
 		{
 			NNE_RUNTIME_BASIC_TRACE_SCOPE(NNE::RuntimeBasic::Private::OperatorGELU);
 
-			check(Output != Input);
+			NNE_RUNTIME_BASIC_CHECK_ALIASING(Output, Input);
 
 #if NNE_RUNTIME_BASIC_ENABLE_ISPC
 			ispc::NNERuntimeBasicCPUOperatorGELU(
@@ -601,7 +610,7 @@ namespace UE::NNE::RuntimeBasic
 		{
 			NNE_RUNTIME_BASIC_TRACE_SCOPE(NNE::RuntimeBasic::Private::OperatorTanH);
 
-			check(Output != Input);
+			NNE_RUNTIME_BASIC_CHECK_ALIASING(Output, Input);
 
 #if NNE_RUNTIME_BASIC_ENABLE_ISPC
 			ispc::NNERuntimeBasicCPUOperatorTanH(
@@ -633,7 +642,7 @@ namespace UE::NNE::RuntimeBasic
 		{
 			NNE_RUNTIME_BASIC_TRACE_SCOPE(NNE::RuntimeBasic::Private::OperatorPReLU);
 
-			check(Output != Input);
+			NNE_RUNTIME_BASIC_CHECK_ALIASING(Output, Input);
 
 #if NNE_RUNTIME_BASIC_ENABLE_ISPC
 			ispc::NNERuntimeBasicCPUOperatorPReLU(
@@ -670,9 +679,9 @@ namespace UE::NNE::RuntimeBasic
 		{
 			NNE_RUNTIME_BASIC_TRACE_SCOPE(NNE::RuntimeBasic::Private::OperatorMemoryCellUpdateMemory);
 
-			check(Output != RememberGate);
-			check(Output != Memory);
-			check(Output != Update);
+			NNE_RUNTIME_BASIC_CHECK_ALIASING(Output, RememberGate);
+			NNE_RUNTIME_BASIC_CHECK_ALIASING(Output, Memory);
+			NNE_RUNTIME_BASIC_CHECK_ALIASING(Output, Update);
 
 #if NNE_RUNTIME_BASIC_ENABLE_ISPC
 			ispc::NNERuntimeBasicCPUOperatorMemoryCellUpdateMemory(
@@ -715,9 +724,9 @@ namespace UE::NNE::RuntimeBasic
 		{
 			NNE_RUNTIME_BASIC_TRACE_SCOPE(NNE::RuntimeBasic::Private::OperatorMemoryCellUpdateOutput);
 
-			check(Output != PassthroughGate);
-			check(Output != MemoryUpdate);
-			check(Output != InputUpdate);
+			NNE_RUNTIME_BASIC_CHECK_ALIASING(Output, PassthroughGate);
+			NNE_RUNTIME_BASIC_CHECK_ALIASING(Output, MemoryUpdate);
+			NNE_RUNTIME_BASIC_CHECK_ALIASING(Output, InputUpdate);
 
 #if NNE_RUNTIME_BASIC_ENABLE_ISPC
 			ispc::NNERuntimeBasicCPUOperatorMemoryCellUpdateOutput(
@@ -757,7 +766,7 @@ namespace UE::NNE::RuntimeBasic
 		{
 			NNE_RUNTIME_BASIC_TRACE_SCOPE(NNE::RuntimeBasic::Private::OperatorAggregateGatherElements);
 
-			check(OutputBuffer != InputBuffer);
+			NNE_RUNTIME_BASIC_CHECK_ALIASING(OutputBuffer, InputBuffer);
 
 #if NNE_RUNTIME_BASIC_ENABLE_ISPC
 			ispc::NNERuntimeBasicCPUOperatorAggregateGatherElements(
@@ -870,11 +879,11 @@ namespace UE::NNE::RuntimeBasic
 		{
 			NNE_RUNTIME_BASIC_TRACE_SCOPE(NNE::RuntimeBasic::Private::OperatorAggregateGatherFromSubLayers);
 
-			check(QueryBuffer != KeyBuffer);
-			check(QueryBuffer != ValueBuffer);
-			check(KeyBuffer != ValueBuffer);
-			check(ElementAccum != ElementNums);
-			check(ElementAccum != ElementOffsets);
+			NNE_RUNTIME_BASIC_CHECK_ALIASING(QueryBuffer, KeyBuffer);
+			NNE_RUNTIME_BASIC_CHECK_ALIASING(QueryBuffer, ValueBuffer);
+			NNE_RUNTIME_BASIC_CHECK_ALIASING(KeyBuffer, ValueBuffer);
+			NNE_RUNTIME_BASIC_CHECK_ALIASING(ElementAccum, ElementNums);
+			NNE_RUNTIME_BASIC_CHECK_ALIASING(ElementAccum, ElementOffsets);
 
 			const uint32 SubLayerNum = SubLayerBatchIndices.Num();
 
@@ -1003,9 +1012,9 @@ namespace UE::NNE::RuntimeBasic
 		{
 			NNE_RUNTIME_BASIC_TRACE_SCOPE(NNE::RuntimeBasic::Private::OperatorAggregateSoftmaxPlusOneInplace);
 
-			check(AttentionMaxs != AttentionDenoms);
-			check(AttentionMaxs != Attention);
-			check(Attention != AttentionDenoms);
+			NNE_RUNTIME_BASIC_CHECK_ALIASING(AttentionMaxs, AttentionDenoms);
+			NNE_RUNTIME_BASIC_CHECK_ALIASING(AttentionMaxs, Attention);
+			NNE_RUNTIME_BASIC_CHECK_ALIASING(Attention, AttentionDenoms);
 
 			// Numerically stable soft-max computation using subtraction of the (positive) max value
 			// 
@@ -1074,8 +1083,8 @@ namespace UE::NNE::RuntimeBasic
 		{
 			NNE_RUNTIME_BASIC_TRACE_SCOPE(NNE::RuntimeBasic::Private::OperatorAggregateAttentionSum);
 
-			check(Output != Attention);
-			check(Output != Values);
+			NNE_RUNTIME_BASIC_CHECK_ALIASING(Output, Attention);
+			NNE_RUNTIME_BASIC_CHECK_ALIASING(Output, Values);
 
 #if NNE_RUNTIME_BASIC_ENABLE_ISPC
 			ispc::NNERuntimeBasicCPUOperatorAggregateAttentionSum(
@@ -1132,7 +1141,7 @@ namespace UE::NNE::RuntimeBasic
 		{
 			NNE_RUNTIME_BASIC_TRACE_SCOPE(NNE::RuntimeBasic::Private::OperatorGather);
 
-			check(OutputBuffer != InputBuffer);
+			NNE_RUNTIME_BASIC_CHECK_ALIASING(OutputBuffer, InputBuffer);
 
 #if NNE_RUNTIME_BASIC_ENABLE_ISPC
 			ispc::NNERuntimeBasicCPUOperatorGather(
@@ -1166,7 +1175,7 @@ namespace UE::NNE::RuntimeBasic
 		{
 			NNE_RUNTIME_BASIC_TRACE_SCOPE(NNE::RuntimeBasic::Private::OperatorScatter);
 
-			check(OutputBuffer != InputBuffer);
+			NNE_RUNTIME_BASIC_CHECK_ALIASING(OutputBuffer, InputBuffer);
 
 #if NNE_RUNTIME_BASIC_ENABLE_ISPC
 			ispc::NNERuntimeBasicCPUOperatorScatter(
@@ -1261,9 +1270,9 @@ namespace UE::NNE::RuntimeBasic
 		{
 			NNE_RUNTIME_BASIC_TRACE_SCOPE(NNE::RuntimeBasic::Private::OperatorGatherTopTwoSubLayerBatchIndices);
 
-			check(BatchSubLayerIndex0 != BatchSubLayerIndex1);
-			check(BatchSubLayerWeight0 != BatchSubLayerWeight1);
-			check(BatchSubLayerOutputIndex0 != BatchSubLayerOutputIndex1);
+			NNE_RUNTIME_BASIC_CHECK_ALIASING(BatchSubLayerIndex0, BatchSubLayerIndex1);
+			NNE_RUNTIME_BASIC_CHECK_ALIASING(BatchSubLayerWeight0, BatchSubLayerWeight1);
+			NNE_RUNTIME_BASIC_CHECK_ALIASING(BatchSubLayerOutputIndex0, BatchSubLayerOutputIndex1);
 
 			check(SubLayerGateSize >= 2);
 
@@ -1328,6 +1337,9 @@ namespace UE::NNE::RuntimeBasic
 		{
 			NNE_RUNTIME_BASIC_TRACE_SCOPE(NNE::RuntimeBasic::Private::OperatorGatherTopTwoFromSubLayers);
 
+			NNE_RUNTIME_BASIC_CHECK_ALIASING(OutputBuffer, BatchSubLayerWeight0);
+			NNE_RUNTIME_BASIC_CHECK_ALIASING(OutputBuffer, BatchSubLayerWeight1);
+
 			for (uint32 BatchIdx = 0; BatchIdx < BatchSize; BatchIdx++)
 			{
 				const uint32 SubLayerIndex0 = BatchSubLayerIndex0[BatchIdx];
@@ -1340,6 +1352,9 @@ namespace UE::NNE::RuntimeBasic
 				const float* RESTRICT SubLayerOutputBuffer0 = SubLayerOutputBuffer[SubLayerIndex0].GetData();
 				const float* RESTRICT SubLayerOutputBuffer1 = SubLayerOutputBuffer[SubLayerIndex1].GetData();
 
+				NNE_RUNTIME_BASIC_CHECK_ALIASING(OutputBuffer, SubLayerOutputBuffer0);
+				NNE_RUNTIME_BASIC_CHECK_ALIASING(OutputBuffer, SubLayerOutputBuffer1);
+
 				for (uint32 OutputIdx = 0; OutputIdx < OutputBufferSize; OutputIdx++)
 				{
 					OutputBuffer[BatchIdx * OutputBufferStride + OutputIdx] =
@@ -1362,7 +1377,7 @@ namespace UE::NNE::RuntimeBasic
 		{
 			NNE_RUNTIME_BASIC_TRACE_SCOPE(NNE::RuntimeBasic::Private::OperatorLayerNorm);
 
-			check(Output != Input);
+			NNE_RUNTIME_BASIC_CHECK_ALIASING(Output, Input);
 
 #if NNE_RUNTIME_BASIC_ENABLE_ISPC
 			ispc::NNERuntimeBasicCPUOperatorLayerNorm(
@@ -5392,7 +5407,9 @@ namespace UE::NNE::RuntimeBasic
 
 } // namespace UE::NNE::RuntimeBasic
 
+#undef NNE_RUNTIME_BASIC_CHECK_ALIASING
 #undef NNE_RUNTIME_BASIC_TRACE_SCOPE
+#undef NNE_RUNTIME_BASIC_ENABLE_ALIASING_CHECK
 #undef NNE_RUNTIME_BASIC_ENABLE_PROFILE
 #undef NNE_RUNTIME_BASIC_ENABLE_NAN_CHECK
 #undef NNE_RUNTIME_BASIC_ENABLE_ISPC
\ No newline at end of file
