LearningAgents: Fixed issue with log prob and entropy computation. Added more NaN checking to NNERuntimeBasicCpu.

diff --git a/Engine/Plugins/Experimental/NNERuntimeBasicCpu/Source/NNERuntimeBasicCpu/Private/NNERuntimeBasicCpuModel.cpp b/Engine/Plugins/Experimental/NNERuntimeBasicCpu/Source/NNERuntimeBasicCpu/Private/NNERuntimeBasicCpuModel.cpp
--- a/Engine/Plugins/Experimental/NNERuntimeBasicCpu/Source/NNERuntimeBasicCpu/Private/NNERuntimeBasicCpuModel.cpp
+++ b/Engine/Plugins/Experimental/NNERuntimeBasicCpu/Source/NNERuntimeBasicCpu/Private/NNERuntimeBasicCpuModel.cpp
@@ -46,6 +46,19 @@ namespace UE::NNE::RuntimeBasic
 		{
 			//--------------------------------------------------------------------------
 
+			static inline void NanCheck(TConstArrayView<float> In)
+			{
+#if NNE_RUNTIME_BASIC_ENABLE_NAN_CHECK
+				for (int32 Idx = 0; Idx < In.Num(); Idx++)
+				{
+					checkf(FMath::IsFinite(In[Idx]) && In[Idx] != MAX_flt && In[Idx] != -MAX_flt,
+						TEXT("Invalid value %f found in Array at index %i"), In[Idx], Idx);
+				}
+#endif
+			}
+
+			//--------------------------------------------------------------------------
+
 			static inline void Align(uint64& InOutOffset, const uint32 Alignment)
 			{
 				InOutOffset = ((InOutOffset + Alignment - 1) / Alignment) * Alignment;
@@ -107,6 +120,7 @@ namespace UE::NNE::RuntimeBasic
 				Align(InOutOffset, 64);
 				Out = MakeArrayView<const float>((const float*)(Data.GetData() + InOutOffset), Size);
 				InOutOffset += Size * sizeof(float);
+				NanCheck(Out);
 			}
 
 			static inline void Load(uint64& InOutOffset, TConstArrayView<uint16>& Out, TConstArrayView<uint8> Data, int32 Size)
@@ -144,6 +158,7 @@ namespace UE::NNE::RuntimeBasic
 
 			static inline void Save(uint64& InOutOffset, TConstArrayView<float> In, TArrayView<uint8> Data)
 			{
+				NanCheck(In);
 				Align(InOutOffset, 64);
 				FMemory::Memcpy(Data.GetData() + InOutOffset, In.GetData(), In.Num() * sizeof(float));
 				InOutOffset += In.Num() * sizeof(float);
