NNERuntimeBasicCpu: Added some new layers to support skip connections

diff --git a/Engine/Plugins/Experimental/NNERuntimeBasicCpu/Source/NNERuntimeBasicCpu/Private/NNERuntimeBasicCpu.ispc b/Engine/Plugins/Experimental/NNERuntimeBasicCpu/Source/NNERuntimeBasicCpu/Private/NNERuntimeBasicCpu.ispc
--- a/Engine/Plugins/Experimental/NNERuntimeBasicCpu/Source/NNERuntimeBasicCpu/Private/NNERuntimeBasicCpu.ispc
+++ b/Engine/Plugins/Experimental/NNERuntimeBasicCpu/Source/NNERuntimeBasicCpu/Private/NNERuntimeBasicCpu.ispc
@@ -191,6 +191,22 @@ export void NNERuntimeBasicCPUOperatorCopy(
 	}
 }
 
+export void NNERuntimeBasicCPUOperatorTile(
+	uniform float Output[],
+	uniform const float Input[],
+	uniform const int BatchNum,
+	uniform const int Num,
+	uniform const int Repeats,
+	uniform const int OutputStride,
+	uniform const int InputStride)
+{
+	foreach(BatchIdx = 0 ... BatchNum, Repeat = 0 ... Repeats, Idx = 0 ... Num)
+	{
+		Output[BatchIdx * OutputStride + Repeat * Num + Idx] = Input[BatchIdx * InputStride + Idx];
+	}
+}
+
+
 export void NNERuntimeBasicCPUOperatorReLU(
 	uniform float Output[],
 	uniform const float Input[],
diff --git a/Engine/Plugins/Experimental/NNERuntimeBasicCpu/Source/NNERuntimeBasicCpu/Private/NNERuntimeBasicCpuModel.cpp b/Engine/Plugins/Experimental/NNERuntimeBasicCpu/Source/NNERuntimeBasicCpu/Private/NNERuntimeBasicCpuModel.cpp
--- a/Engine/Plugins/Experimental/NNERuntimeBasicCpu/Source/NNERuntimeBasicCpu/Private/NNERuntimeBasicCpuModel.cpp
+++ b/Engine/Plugins/Experimental/NNERuntimeBasicCpu/Source/NNERuntimeBasicCpu/Private/NNERuntimeBasicCpuModel.cpp
@@ -248,6 +248,42 @@ namespace UE::NNE::RuntimeBasic
 #endif
 		}
 
+		static inline void OperatorTile(
+			float* RESTRICT Output,
+			const float* RESTRICT Input,
+			const uint32 BatchSize,
+			const uint32 InputSize,
+			const uint32 Repeats,
+			const uint32 OutputStride,
+			const uint32 InputStride)
+		{
+			NNE_RUNTIME_BASIC_TRACE_SCOPE(NNE::RuntimeBasic::Private::OperatorTile);
+
+			NNE_RUNTIME_BASIC_CHECK_ALIASING(Output, Input);
+
+#if NNE_RUNTIME_BASIC_ENABLE_ISPC
+			ispc::NNERuntimeBasicCPUOperatorTile(
+				Output,
+				Input,
+				BatchSize,
+				InputSize,
+				Repeats,
+				OutputStride,
+				InputStride);
+#else
+			for (uint32 BatchIdx = 0; BatchIdx < BatchSize; BatchIdx++)
+			{
+				for (uint32 Repeat = 0; Repeat < Repeats; Repeat++)
+				{
+					for (uint32 Idx = 0; Idx < InputSize; Idx++)
+					{
+						Output[BatchIdx * OutputStride + Repeats * InputSize + Idx] = Input[BatchIdx * InputStride + Idx];
+					}
+				}
+			}
+#endif
+		}
+
 		static inline void OperatorNormalize(
 			float* RESTRICT Output,
 			const float* RESTRICT Input,
@@ -1459,6 +1495,9 @@ namespace UE::NNE::RuntimeBasic
 			GELU = 20,
 			LayerNorm = 21,
 			LipschiztLinear = 22,
+			Tile = 23,
+			Spread = 24,
+			Slice = 25,
 		};
 
 		//--------------------------------------------------------------------------
@@ -4109,6 +4148,263 @@ namespace UE::NNE::RuntimeBasic
 			TConstArrayView<float> Weights;
 		};
 
+
+		//--------------------------------------------------------------------------
+
+		struct FTileLayer : public ILayer
+		{
+			virtual ELayerType GetLayerType() const override final { return ELayerType::Tile; }
+			virtual uint32 GetInputSize() const override final { return InputSize; }
+			virtual uint32 GetOutputSize() const override final { return Repeats * InputSize; }
+
+			virtual void SerializationSize(uint64& InOutOffset) const override final
+			{
+				Serialization::Size(InOutOffset, InputSize);
+				Serialization::Size(InOutOffset, Repeats);
+			}
+
+			virtual void SerializationLoad(uint64& InOutOffset, TConstArrayView<uint8> Data) override final
+			{
+				Serialization::Load(InOutOffset, InputSize, Data);
+				Serialization::Load(InOutOffset, Repeats, Data);
+			}
+
+			virtual void SerializationSave(uint64& InOutOffset, TArrayView<uint8> Data) const override final
+			{
+				Serialization::Save(InOutOffset, InputSize, Data);
+				Serialization::Save(InOutOffset, Repeats, Data);
+			}
+
+			virtual void Evaluate(
+				ILayerInstance* Instance,
+				float* OutputBuffer,
+				const float* InputBuffer,
+				const uint32 BatchSize,
+				const uint32 OutputBufferSize,
+				const uint32 InputBufferSize,
+				const uint32 OutputBufferStride,
+				const uint32 InputBufferStride) override final
+			{
+				NNE_RUNTIME_BASIC_TRACE_SCOPE(NNE::RuntimeBasic::Private::FTileLayer::Evaluate);
+				check(OutputBufferSize == GetOutputSize() && InputBufferSize == GetInputSize());
+				check(OutputBufferStride >= GetOutputSize() && InputBufferStride >= GetInputSize());
+				check(Instance == nullptr);
+				OperatorNanCheck(InputBuffer, BatchSize, InputBufferSize, InputBufferStride);
+
+				OperatorTile(
+					OutputBuffer,
+					InputBuffer,
+					BatchSize,
+					InputSize,
+					Repeats,
+					OutputBufferStride,
+					InputBufferStride);
+
+				OperatorNanCheck(OutputBuffer, BatchSize, OutputBufferSize, OutputBufferStride);
+			}
+
+			uint32 Repeats = 1;
+			uint32 InputSize = 0;
+		};
+
+		//--------------------------------------------------------------------------
+
+		struct FSpreadLayer;
+
+		struct FSpreadLayerInstance : public ILayerInstance
+		{
+			FSpreadLayerInstance(const FSpreadLayer& InSpreadLayer);
+
+			void SetMaxBatchSize(const uint32 MaxBatchSize) override final;
+
+			const FSpreadLayer& SpreadLayer;
+			TArray<TSharedPtr<ILayerInstance>, TInlineAllocator<32>> Instances;
+		};
+
+		struct FSpreadLayer : public ILayer
+		{
+			virtual TSharedPtr<ILayerInstance> MakeInstance() const { return MakeShared<FSpreadLayerInstance>(*this); };
+			virtual ELayerType GetLayerType() const override final { return ELayerType::Spread; }
+			virtual uint32 GetInputSize() const override final { return InputSize; }
+			virtual uint32 GetOutputSize() const override final { return TotalOutputSize; }
+
+			virtual void SerializationSize(uint64& InOutOffset) const override final
+			{
+				Serialization::Size(InOutOffset, (uint32)Layers.Num());
+				Serialization::Size(InOutOffset, InputSize);
+				Serialization::Size(InOutOffset, OutputSizes);
+				Serialization::Size(InOutOffset, Layers);
+			}
+
+			virtual void SerializationLoad(uint64& InOutOffset, TConstArrayView<uint8> Data) override final
+			{
+				uint32 LayerNum = 0;
+				Serialization::Load(InOutOffset, LayerNum, Data);
+				Serialization::Load(InOutOffset, InputSize, Data);
+				Serialization::Load(InOutOffset, OutputSizes, Data, LayerNum);
+				Layers.Init(nullptr, LayerNum);
+				Serialization::Load(InOutOffset, Layers, Data);
+				PostLoad();
+			}
+
+			void PostLoad()
+			{
+				const uint32 LayerNum = Layers.Num();
+				OutputOffsets.SetNumUninitialized(LayerNum);
+
+				TotalOutputSize = 0;
+				for (uint32 LayerIdx = 0; LayerIdx < LayerNum; LayerIdx++)
+				{
+					OutputOffsets[LayerIdx] = TotalOutputSize;
+					TotalOutputSize += OutputSizes[LayerIdx];
+				}
+			}
+
+			virtual void SerializationSave(uint64& InOutOffset, TArrayView<uint8> Data) const override final
+			{
+				Serialization::Save(InOutOffset, (uint32)Layers.Num(), Data);
+				Serialization::Save(InOutOffset, InputSize, Data);
+				Serialization::Save(InOutOffset, OutputSizes, Data);
+				Serialization::Save(InOutOffset, Layers, Data);
+			}
+
+			virtual void Evaluate(
+				ILayerInstance* Instance,
+				float* OutputBuffer,
+				const float* InputBuffer,
+				const uint32 BatchSize,
+				const uint32 OutputBufferSize,
+				const uint32 InputBufferSize,
+				const uint32 OutputBufferStride,
+				const uint32 InputBufferStride) override final
+			{
+				NNE_RUNTIME_BASIC_TRACE_SCOPE(NNE::RuntimeBasic::Private::FSpreadLayer::Evaluate);
+				check(OutputBufferSize == GetOutputSize() && InputBufferSize == GetInputSize());
+				check(OutputBufferStride >= GetOutputSize() && InputBufferStride >= GetInputSize());
+
+				FSpreadLayerInstance* SpreadInstance = StaticCast<FSpreadLayerInstance*>(Instance);
+				check(SpreadInstance);
+
+				OperatorNanCheck(InputBuffer, BatchSize, InputBufferSize, InputBufferStride);
+
+				const int32 LayerNum = Layers.Num();
+
+				for (int32 LayerIdx = 0; LayerIdx < LayerNum; LayerIdx++)
+				{
+					check(OutputOffsets[LayerIdx] + OutputSizes[LayerIdx] <= OutputBufferSize);
+
+					Layers[LayerIdx]->Evaluate(
+						SpreadInstance->Instances[LayerIdx].Get(),
+						OutputBuffer + OutputOffsets[LayerIdx],
+						InputBuffer,
+						BatchSize,
+						OutputSizes[LayerIdx],
+						InputSize,
+						OutputBufferStride,
+						InputBufferStride);
+				}
+
+				OperatorNanCheck(OutputBuffer, BatchSize, OutputBufferSize, OutputBufferStride);
+			}
+
+			TConstArrayView<uint32> OutputSizes;
+			TArray<TSharedPtr<ILayer>, TInlineAllocator<32>> Layers;
+
+			uint32 InputSize = 0;
+			uint32 TotalOutputSize = 0;
+			TArray<uint32, TInlineAllocator<32>> OutputOffsets;
+		};
+
+		FSpreadLayerInstance::FSpreadLayerInstance(const FSpreadLayer& InSpreadLayer)
+			: SpreadLayer(InSpreadLayer)
+		{
+			const uint32 LayerNum = SpreadLayer.Layers.Num();
+			Instances.Init(nullptr, LayerNum);
+
+			for (uint32 LayerIdx = 0; LayerIdx < LayerNum; LayerIdx++)
+			{
+				Instances[LayerIdx] = SpreadLayer.Layers[LayerIdx]->MakeInstance();
+			}
+		}
+
+		void FSpreadLayerInstance::SetMaxBatchSize(const uint32 MaxBatchSize)
+		{
+			NNE_RUNTIME_BASIC_TRACE_SCOPE(NNE::RuntimeBasic::Private::FSpreadLayerInstance::SetMaxBatchSize);
+
+			// Propagate call to sub-layer instances
+
+			for (const TSharedPtr<ILayerInstance>& Instance : Instances)
+			{
+				if (Instance)
+				{
+					Instance->SetMaxBatchSize(MaxBatchSize);
+				}
+			}
+		}
+
+
+		//--------------------------------------------------------------------------
+
+		struct FSliceLayer : public ILayer
+		{
+			virtual ELayerType GetLayerType() const override final { return ELayerType::Slice; }
+			virtual uint32 GetInputSize() const override final { return InputSize; }
+			virtual uint32 GetOutputSize() const override final { return SliceSize; }
+
+			virtual void SerializationSize(uint64& InOutOffset) const override final
+			{
+				Serialization::Size(InOutOffset, InputSize);
+				Serialization::Size(InOutOffset, SliceOffset);
+				Serialization::Size(InOutOffset, SliceSize);
+			}
+
+			virtual void SerializationLoad(uint64& InOutOffset, TConstArrayView<uint8> Data) override final
+			{
+				Serialization::Load(InOutOffset, InputSize, Data);
+				Serialization::Load(InOutOffset, SliceOffset, Data);
+				Serialization::Load(InOutOffset, SliceSize, Data);
+			}
+
+			virtual void SerializationSave(uint64& InOutOffset, TArrayView<uint8> Data) const override final
+			{
+				Serialization::Save(InOutOffset, InputSize, Data);
+				Serialization::Save(InOutOffset, SliceOffset, Data);
+				Serialization::Save(InOutOffset, SliceSize, Data);
+			}
+
+			virtual void Evaluate(
+				ILayerInstance* Instance,
+				float* OutputBuffer,
+				const float* InputBuffer,
+				const uint32 BatchSize,
+				const uint32 OutputBufferSize,
+				const uint32 InputBufferSize,
+				const uint32 OutputBufferStride,
+				const uint32 InputBufferStride) override final
+			{
+				NNE_RUNTIME_BASIC_TRACE_SCOPE(NNE::RuntimeBasic::Private::FSliceLayer::Evaluate);
+				check(OutputBufferSize == GetOutputSize() && InputBufferSize == GetInputSize());
+				check(OutputBufferStride >= GetOutputSize() && InputBufferStride >= GetInputSize());
+				check(Instance == nullptr);
+				check(SliceOffset + SliceSize <= InputBufferSize);
+				OperatorNanCheck(InputBuffer, BatchSize, InputBufferSize, InputBufferStride);
+
+				OperatorCopy(
+					OutputBuffer,
+					InputBuffer + SliceOffset,
+					BatchSize,
+					SliceSize,
+					OutputBufferStride,
+					InputBufferStride);
+
+				OperatorNanCheck(OutputBuffer, BatchSize, OutputBufferSize, OutputBufferStride);
+			}
+
+			uint32 InputSize = 0;
+			uint32 SliceOffset = 0;
+			uint32 SliceSize = 0;
+		};
+
 		//--------------------------------------------------------------------------
 		// Layer Serialization
 		//--------------------------------------------------------------------------
@@ -4159,6 +4455,9 @@ namespace UE::NNE::RuntimeBasic
 					case ELayerType::GELU: OutLayer = MakeShared<FGELULayer>(); break;
 					case ELayerType::LayerNorm: OutLayer = MakeShared<FLayerNormLayer>(); break;
 					case ELayerType::LipschiztLinear: OutLayer = MakeShared<FLipschiztLinearLayer>(); break;
+					case ELayerType::Tile: OutLayer = MakeShared<FTileLayer>(); break;
+					case ELayerType::Spread: OutLayer = MakeShared<FSpreadLayer>(); break;
+					case ELayerType::Slice: OutLayer = MakeShared<FSliceLayer>(); break;
 					default: checkf(false, TEXT("Unknown Layer Id %i"), LayerTypeId);
 					}
 				}
@@ -4602,6 +4901,17 @@ namespace UE::NNE::RuntimeBasic
 		return StaticCastSharedPtr<Private::ILayer>(Layer);
 	}
 
+	FModelBuilderElement FModelBuilder::MakeSlice(const uint32 InputSize, const uint32 SliceOffset, const uint32 SliceSize)
+	{
+		check(SliceOffset + SliceSize <= InputSize);
+
+		const TSharedPtr<Private::FSliceLayer> Layer = MakeShared<Private::FSliceLayer>();
+		Layer->InputSize = InputSize;
+		Layer->SliceOffset = SliceOffset;
+		Layer->SliceSize = SliceSize;
+		return StaticCastSharedPtr<Private::ILayer>(Layer);
+	}
+
 	FModelBuilderElement FModelBuilder::MakeClamp(const uint32 InputOutputSize, const TConstArrayView<float> MinValues, const TConstArrayView<float> MaxValues)
 	{
 		check(MinValues.Num() == InputOutputSize);
@@ -4738,6 +5048,65 @@ namespace UE::NNE::RuntimeBasic
 		return MakeSequence(Layers);
 	}
 
+	FModelBuilderElement FModelBuilder::MakeSkipMLP(
+		const uint32 InputSize,
+		const uint32 OutputSize,
+		const uint32 HiddenSize,
+		const uint32 LayerNum,
+		const EActivationFunction ActivationFunction,
+		const bool bActivationOnFinalLayer,
+		const FLinearLayerSettings& LinearLayerSettings)
+	{
+		check(LayerNum >= 2);
+
+		const int32 TotalLayerNum = LayerNum;
+
+		TArray<FModelBuilderElement, TInlineAllocator<32>> Layers;
+		Layers.Reserve(TotalLayerNum);
+
+		for (uint32 LayerIdx = 0; LayerIdx < LayerNum; LayerIdx++)
+		{
+			if (LayerIdx == 0)
+			{
+				Layers.Emplace(MakeSpread({
+					MakeCopy(InputSize),
+					MakeSequence({
+						MakeLinearLayer(InputSize, HiddenSize, LinearLayerSettings),
+						MakeActivation(HiddenSize, ActivationFunction),
+					})
+					}));
+			}
+			else if (LayerIdx != LayerNum - 1)
+			{
+				Layers.Emplace(MakeSpread({
+					MakeSlice(InputSize + HiddenSize, 0, InputSize),
+					MakeSequence({
+						MakeLinearLayer(InputSize + HiddenSize, HiddenSize, LinearLayerSettings),
+						MakeActivation(HiddenSize, ActivationFunction),
+					})
+					}));
+			}
+			else
+			{
+				if (bActivationOnFinalLayer)
+				{
+					Layers.Emplace(MakeSequence({
+							MakeLinearLayer(InputSize + HiddenSize, OutputSize, LinearLayerSettings),
+							MakeActivation(OutputSize, ActivationFunction),
+						}));
+				}
+				else
+				{
+					Layers.Emplace(MakeLinearLayer(InputSize + HiddenSize, OutputSize, LinearLayerSettings));
+				}
+			}
+		}
+
+		check(Layers.Num() == TotalLayerNum);
+
+		return MakeSequence(Layers);
+	}
+
 	FModelBuilderElement FModelBuilder::MakeMemoryCell(
 		const uint32 InputNum,
 		const uint32 OutputNum,
@@ -4830,6 +5199,26 @@ namespace UE::NNE::RuntimeBasic
 		return StaticCastSharedPtr<Private::ILayer>(ConcatLayer);
 	}
 
+	FModelBuilderElement FModelBuilder::MakeSpread(const TConstArrayView<FModelBuilderElement> Elements)
+	{
+		const int32 LayerNum = Elements.Num();
+
+		const TSharedPtr<Private::FSpreadLayer> SpreadLayer = MakeShared<Private::FSpreadLayer>();
+		SpreadLayer->InputSize = LayerNum ? Elements[0].GetInputSize() : 0;
+		SpreadLayer->OutputSizes = MakeSizesLayerOutputs(Elements);
+		SpreadLayer->Layers.Reserve(LayerNum);
+
+		for (const FModelBuilderElement& Element : Elements)
+		{
+			check(SpreadLayer->InputSize == Element.GetInputSize());
+			SpreadLayer->Layers.Emplace(Element.Layer);
+		}
+
+		SpreadLayer->PostLoad();
+
+		return StaticCastSharedPtr<Private::ILayer>(SpreadLayer);
+	}
+
 	FModelBuilderElement FModelBuilder::MakeArray(const uint32 ElementNum, const FModelBuilderElement& SubLayer)
 	{
 		const TSharedPtr<Private::FArrayLayer> ArrayLayer = MakeShared<Private::FArrayLayer>();
@@ -5001,6 +5390,14 @@ namespace UE::NNE::RuntimeBasic
 		return StaticCastSharedPtr<Private::ILayer>(LayerNormLayer);
 	}
 
+	FModelBuilderElement FModelBuilder::MakeTile(const uint32 InputSize, const uint32 Repeats)
+	{
+		const TSharedPtr<Private::FTileLayer> Layer = MakeShared<Private::FTileLayer>();
+		Layer->InputSize = InputSize;
+		Layer->Repeats = Repeats;
+		return StaticCastSharedPtr<Private::ILayer>(Layer);
+	}
+
 	void FModelBuilder::Reset()
 	{
 		Rng = RngInitialState;
diff --git a/Engine/Plugins/Experimental/NNERuntimeBasicCpu/Source/NNERuntimeBasicCpu/Public/NNERuntimeBasicCpuBuilder.h b/Engine/Plugins/Experimental/NNERuntimeBasicCpu/Source/NNERuntimeBasicCpu/Public/NNERuntimeBasicCpuBuilder.h
--- a/Engine/Plugins/Experimental/NNERuntimeBasicCpu/Source/NNERuntimeBasicCpu/Public/NNERuntimeBasicCpuBuilder.h
+++ b/Engine/Plugins/Experimental/NNERuntimeBasicCpu/Source/NNERuntimeBasicCpu/Public/NNERuntimeBasicCpuBuilder.h
@@ -224,6 +224,9 @@ namespace UE::NNE::RuntimeBasic
 		/** Makes a Copy Layer */
 		FModelBuilderElement MakeCopy(const uint32 InputOutputSize);
 
+		/** Makes a Slice Layer */
+		FModelBuilderElement MakeSlice(const uint32 InputSize, const uint32 SliceOffset, const uint32 SliceSize);
+
 		/** Makes a Clamp Layer */
 		FModelBuilderElement MakeClamp(const uint32 InputOutputSize, const TConstArrayView<float> MinValues, const TConstArrayView<float> MaxValues);
 
@@ -283,6 +286,26 @@ namespace UE::NNE::RuntimeBasic
 			const bool bActivationOnFinalLayer = false,
 			const FLinearLayerSettings& LinearLayerSettings = FLinearLayerSettings());
 
+		/**
+		 * Makes a Multi-Layer Perceptron network using Skip layers that concatenate the input to each intermediate layer.
+		 *
+		 * @param InputSize					Input Vector Size
+		 * @param OutputSize				Output Vector Size
+		 * @param HiddenSize				Number of hidden units to use on internal layers.
+		 * @param LayerNum					Number of layers. Includes input and output layers.
+		 * @param ActivationFunction		Activation function to use.
+		 * @param bActivationOnFinalLayer	If the activation function should be used on the final output layer.
+		 * @param LinearLayerSettings		Settings for the linear layers in the MLP.
+		 */
+		FModelBuilderElement MakeSkipMLP(
+			const uint32 InputSize,
+			const uint32 OutputSize,
+			const uint32 HiddenSize,
+			const uint32 LayerNum,
+			const EActivationFunction ActivationFunction,
+			const bool bActivationOnFinalLayer = false,
+			const FLinearLayerSettings& LinearLayerSettings = FLinearLayerSettings());
+
 		/**
 		 * Make a new Memory Cell layer.
 		 *
@@ -339,6 +362,12 @@ namespace UE::NNE::RuntimeBasic
 		 */
 		FModelBuilderElement MakeConcat(const TConstArrayView<FModelBuilderElement> Elements);
 		
+		/**
+		 * Makes a Spread layer, which will evaluate each of the given elements on the input vector, 
+		 * concatenating the result into the output vector.
+		 */
+		FModelBuilderElement MakeSpread(const TConstArrayView<FModelBuilderElement> Elements);
+
 		/**
 		 * Make a layer which runs the given sublayer on an array of elements.
 		 * 
@@ -408,6 +437,10 @@ namespace UE::NNE::RuntimeBasic
 			const TConstArrayView<float> Scales,
 			const float Epsilon = 1e-5f);
 
+		/** Makes a Tile Layer */
+		FModelBuilderElement MakeTile(const uint32 InputSize, const uint32 Repeats);
+
+
 	public:
 
 		/** Creates a array of values from a copy of the given array view */
