NNERuntimeBasicCpu: Added SkipMLP with LayerNorm.

diff --git a/Engine/Plugins/Experimental/NNERuntimeBasicCpu/Source/NNERuntimeBasicCpu/Private/NNERuntimeBasicCpuModel.cpp b/Engine/Plugins/Experimental/NNERuntimeBasicCpu/Source/NNERuntimeBasicCpu/Private/NNERuntimeBasicCpuModel.cpp
--- a/Engine/Plugins/Experimental/NNERuntimeBasicCpu/Source/NNERuntimeBasicCpu/Private/NNERuntimeBasicCpuModel.cpp
+++ b/Engine/Plugins/Experimental/NNERuntimeBasicCpu/Source/NNERuntimeBasicCpu/Private/NNERuntimeBasicCpuModel.cpp
@@ -5106,6 +5106,69 @@ namespace UE::NNE::RuntimeBasic
 
 		return MakeSequence(Layers);
 	}
+	
+	FModelBuilderElement FModelBuilder::MakeSkipMLPWithLayerNorm(
+		const uint32 InputSize,
+		const uint32 OutputSize,
+		const uint32 HiddenSize,
+		const uint32 LayerNum,
+		const EActivationFunction ActivationFunction,
+		const bool bActivationOnFinalLayer,
+		const FLinearLayerSettings& LinearLayerSettings)
+	{
+		check(LayerNum >= 2);
+
+		const int32 TotalLayerNum = LayerNum;
+
+		TArray<FModelBuilderElement, TInlineAllocator<32>> Layers;
+		Layers.Reserve(TotalLayerNum);
+
+		for (uint32 LayerIdx = 0; LayerIdx < LayerNum; LayerIdx++)
+		{
+			if (LayerIdx == 0)
+			{
+				Layers.Emplace(MakeSpread({
+					MakeCopy(InputSize),
+					MakeSequence({
+						MakeLinearLayer(InputSize, HiddenSize, LinearLayerSettings)
+					})
+					}));
+			}
+			else if (LayerIdx != LayerNum - 1)
+			{
+				Layers.Emplace(MakeSpread({
+					MakeSlice(InputSize + HiddenSize, 0, InputSize),
+					MakeSequence({
+						MakeLayerNorm(InputSize + HiddenSize, MakeValuesZero(InputSize + HiddenSize), MakeValuesOne(InputSize + HiddenSize)),
+						MakeLinearLayer(InputSize + HiddenSize, HiddenSize, LinearLayerSettings),
+						MakeActivation(HiddenSize, ActivationFunction),
+					})
+					}));
+			}
+			else
+			{
+				if (bActivationOnFinalLayer)
+				{
+					Layers.Emplace(MakeSequence({
+							MakeLayerNorm(InputSize + HiddenSize, MakeValuesZero(InputSize + HiddenSize), MakeValuesOne(InputSize + HiddenSize)),
+							MakeLinearLayer(InputSize + HiddenSize, OutputSize, LinearLayerSettings),
+							MakeActivation(OutputSize, ActivationFunction),
+						}));
+				}
+				else
+				{
+					Layers.Emplace(MakeSequence({
+							MakeLayerNorm(InputSize + HiddenSize, MakeValuesZero(InputSize + HiddenSize), MakeValuesOne(InputSize + HiddenSize)),
+							MakeLinearLayer(InputSize + HiddenSize, OutputSize, LinearLayerSettings)
+						}));
+				}
+			}
+		}
+
+		check(Layers.Num() == TotalLayerNum);
+
+		return MakeSequence(Layers);
+	}
 
 	FModelBuilderElement FModelBuilder::MakeMemoryCell(
 		const uint32 InputNum,
diff --git a/Engine/Plugins/Experimental/NNERuntimeBasicCpu/Source/NNERuntimeBasicCpu/Public/NNERuntimeBasicCpuBuilder.h b/Engine/Plugins/Experimental/NNERuntimeBasicCpu/Source/NNERuntimeBasicCpu/Public/NNERuntimeBasicCpuBuilder.h
--- a/Engine/Plugins/Experimental/NNERuntimeBasicCpu/Source/NNERuntimeBasicCpu/Public/NNERuntimeBasicCpuBuilder.h
+++ b/Engine/Plugins/Experimental/NNERuntimeBasicCpu/Source/NNERuntimeBasicCpu/Public/NNERuntimeBasicCpuBuilder.h
@@ -306,6 +306,26 @@ namespace UE::NNE::RuntimeBasic
 			const bool bActivationOnFinalLayer = false,
 			const FLinearLayerSettings& LinearLayerSettings = FLinearLayerSettings());
 
+		/**
+		 * Makes a Multi-Layer Perceptron network using Skip layers that concatenate the input to each intermediate layer followed by a LayerNorm.
+		 *
+		 * @param InputSize					Input Vector Size
+		 * @param OutputSize				Output Vector Size
+		 * @param HiddenSize				Number of hidden units to use on internal layers.
+		 * @param LayerNum					Number of layers. Includes input and output layers.
+		 * @param ActivationFunction		Activation function to use.
+		 * @param bActivationOnFinalLayer	If the activation function should be used on the final output layer.
+		 * @param LinearLayerSettings		Settings for the linear layers in the MLP.
+		 */
+		FModelBuilderElement MakeSkipMLPWithLayerNorm(
+			const uint32 InputSize,
+			const uint32 OutputSize,
+			const uint32 HiddenSize,
+			const uint32 LayerNum,
+			const EActivationFunction ActivationFunction,
+			const bool bActivationOnFinalLayer = false,
+			const FLinearLayerSettings& LinearLayerSettings = FLinearLayerSettings());
+
 		/**
 		 * Make a new Memory Cell layer.
 		 *
