NNERuntimeBasicCpu: Added Residual Layer

diff --git a/Engine/Plugins/Experimental/NNERuntimeBasicCpu/Source/NNERuntimeBasicCpu/Private/NNERuntimeBasicCpu.ispc b/Engine/Plugins/Experimental/NNERuntimeBasicCpu/Source/NNERuntimeBasicCpu/Private/NNERuntimeBasicCpu.ispc
--- a/Engine/Plugins/Experimental/NNERuntimeBasicCpu/Source/NNERuntimeBasicCpu/Private/NNERuntimeBasicCpu.ispc
+++ b/Engine/Plugins/Experimental/NNERuntimeBasicCpu/Source/NNERuntimeBasicCpu/Private/NNERuntimeBasicCpu.ispc
@@ -191,6 +191,20 @@ export void NNERuntimeBasicCPUOperatorCopy(
 	}
 }
 
+export void NNERuntimeBasicCPUOperatorAddInplace(
+	uniform float Output[],
+	uniform const float Input[],
+	uniform const int BatchNum,
+	uniform const int Num,
+	uniform const int OutputStride,
+	uniform const int InputStride)
+{
+	foreach(BatchIdx = 0 ... BatchNum, Idx = 0 ... Num)
+	{
+		Output[BatchIdx * OutputStride + Idx] += Input[BatchIdx * InputStride + Idx];
+	}
+}
+
 export void NNERuntimeBasicCPUOperatorTile(
 	uniform float Output[],
 	uniform const float Input[],
diff --git a/Engine/Plugins/Experimental/NNERuntimeBasicCpu/Source/NNERuntimeBasicCpu/Private/NNERuntimeBasicCpuModel.cpp b/Engine/Plugins/Experimental/NNERuntimeBasicCpu/Source/NNERuntimeBasicCpu/Private/NNERuntimeBasicCpuModel.cpp
--- a/Engine/Plugins/Experimental/NNERuntimeBasicCpu/Source/NNERuntimeBasicCpu/Private/NNERuntimeBasicCpuModel.cpp
+++ b/Engine/Plugins/Experimental/NNERuntimeBasicCpu/Source/NNERuntimeBasicCpu/Private/NNERuntimeBasicCpuModel.cpp
@@ -248,6 +248,37 @@ namespace UE::NNE::RuntimeBasic
 #endif
 		}
 
+		static inline void OperatorAddInplace(
+			float* RESTRICT Output,
+			const float* RESTRICT Input,
+			const uint32 BatchSize,
+			const uint32 InputOutputSize,
+			const uint32 OutputStride,
+			const uint32 InputStride)
+		{
+			NNE_RUNTIME_BASIC_TRACE_SCOPE(NNE::RuntimeBasic::Private::OperatorAddInplace);
+
+			NNE_RUNTIME_BASIC_CHECK_ALIASING(Output, Input);
+
+#if NNE_RUNTIME_BASIC_ENABLE_ISPC
+			ispc::NNERuntimeBasicCPUOperatorAddInplace(
+				Output,
+				Input,
+				BatchSize,
+				InputOutputSize,
+				OutputStride,
+				InputStride);
+#else
+			for (uint32 BatchIdx = 0; BatchIdx < BatchSize; BatchIdx++)
+			{
+				for (uint32 Idx = 0; Idx < InputOutputSize; Idx++)
+				{
+					Output[BatchIdx * OutputStride + Idx] += Input[BatchIdx * InputStride + Idx];
+				}
+			}
+#endif
+		}
+
 		static inline void OperatorTile(
 			float* RESTRICT Output,
 			const float* RESTRICT Input,
@@ -1498,6 +1529,7 @@ namespace UE::NNE::RuntimeBasic
 			Tile = 23,
 			Spread = 24,
 			Slice = 25,
+			Residual = 26,
 		};
 
 		//--------------------------------------------------------------------------
@@ -4405,6 +4437,75 @@ namespace UE::NNE::RuntimeBasic
 			uint32 SliceSize = 0;
 		};
 
+
+		//--------------------------------------------------------------------------
+
+		struct FResidualLayer : public ILayer
+		{
+			virtual TSharedPtr<ILayerInstance> MakeInstance() const { return SubLayer->MakeInstance(); };
+			virtual ELayerType GetLayerType() const override final { return ELayerType::Residual; }
+			virtual uint32 GetInputSize() const override final { return InputOutputSize; }
+			virtual uint32 GetOutputSize() const override final { return InputOutputSize; }
+
+			virtual void SerializationSize(uint64& InOutOffset) const override final
+			{
+				Serialization::Size(InOutOffset, InputOutputSize);
+				Serialization::Size(InOutOffset, SubLayer);
+			}
+
+			virtual void SerializationLoad(uint64& InOutOffset, TConstArrayView<uint8> Data) override final
+			{
+				Serialization::Load(InOutOffset, InputOutputSize, Data);
+				Serialization::Load(InOutOffset, SubLayer, Data);
+			}
+
+			virtual void SerializationSave(uint64& InOutOffset, TArrayView<uint8> Data) const override final
+			{
+				Serialization::Save(InOutOffset, InputOutputSize, Data);
+				Serialization::Save(InOutOffset, SubLayer, Data);
+			}
+
+			virtual void Evaluate(
+				ILayerInstance* Instance,
+				float* OutputBuffer,
+				const float* InputBuffer,
+				const uint32 BatchSize,
+				const uint32 OutputBufferSize,
+				const uint32 InputBufferSize,
+				const uint32 OutputBufferStride,
+				const uint32 InputBufferStride) override final
+			{
+				NNE_RUNTIME_BASIC_TRACE_SCOPE(NNE::RuntimeBasic::Private::FResidualLayer::Evaluate);
+				check(OutputBufferSize == GetOutputSize() && InputBufferSize == GetInputSize());
+				check(OutputBufferStride >= GetOutputSize() && InputBufferStride >= GetInputSize());
+
+				OperatorNanCheck(InputBuffer, BatchSize, InputBufferSize, InputBufferStride);
+
+				SubLayer->Evaluate(
+					Instance,
+					OutputBuffer,
+					InputBuffer,
+					BatchSize,
+					OutputBufferSize,
+					InputBufferSize,
+					OutputBufferStride,
+					InputBufferStride);
+
+				OperatorAddInplace(
+					OutputBuffer,
+					InputBuffer,
+					BatchSize,
+					InputOutputSize,
+					OutputBufferStride,
+					InputBufferStride);
+
+				OperatorNanCheck(OutputBuffer, BatchSize, OutputBufferSize, OutputBufferStride);
+			}
+
+			uint32 InputOutputSize = 0;
+			TSharedPtr<ILayer> SubLayer;
+		};
+
 		//--------------------------------------------------------------------------
 		// Layer Serialization
 		//--------------------------------------------------------------------------
@@ -4458,6 +4559,7 @@ namespace UE::NNE::RuntimeBasic
 					case ELayerType::Tile: OutLayer = MakeShared<FTileLayer>(); break;
 					case ELayerType::Spread: OutLayer = MakeShared<FSpreadLayer>(); break;
 					case ELayerType::Slice: OutLayer = MakeShared<FSliceLayer>(); break;
+					case ELayerType::Residual: OutLayer = MakeShared<FResidualLayer>(); break;
 					default: checkf(false, TEXT("Unknown Layer Id %i"), LayerTypeId);
 					}
 				}
@@ -5170,6 +5272,121 @@ namespace UE::NNE::RuntimeBasic
 		return MakeSequence(Layers);
 	}
 
+	FModelBuilderElement FModelBuilder::MakeResidualMLP(
+		const uint32 InputSize,
+		const uint32 OutputSize,
+		const uint32 HiddenSize,
+		const uint32 LayerNum,
+		const EActivationFunction ActivationFunction,
+		const bool bActivationOnFinalLayer,
+		const FLinearLayerSettings& LinearLayerSettings)
+	{
+		check(LayerNum >= 2);
+
+		const int32 TotalLayerNum = LayerNum;
+
+		TArray<FModelBuilderElement, TInlineAllocator<32>> Layers;
+		Layers.Reserve(TotalLayerNum);
+
+		for (uint32 LayerIdx = 0; LayerIdx < LayerNum; LayerIdx++)
+		{
+			if (LayerIdx == 0)
+			{
+				Layers.Emplace(
+					MakeSequence({
+						MakeLinearLayer(InputSize, HiddenSize, LinearLayerSettings),
+						MakeActivation(HiddenSize, ActivationFunction),
+					}));
+			}
+			else if (LayerIdx != LayerNum - 1)
+			{
+				Layers.Emplace(MakeResidual(
+					MakeSequence({
+						MakeLinearLayer(HiddenSize, HiddenSize, LinearLayerSettings),
+						MakeActivation(HiddenSize, ActivationFunction),
+					})));
+			}
+			else
+			{
+				if (bActivationOnFinalLayer)
+				{
+					Layers.Emplace(MakeSequence({
+							MakeLinearLayer(HiddenSize, OutputSize, LinearLayerSettings),
+							MakeActivation(OutputSize, ActivationFunction),
+						}));
+				}
+				else
+				{
+					Layers.Emplace(MakeLinearLayer(HiddenSize, OutputSize, LinearLayerSettings));
+				}
+			}
+		}
+
+		check(Layers.Num() == TotalLayerNum);
+
+		return MakeSequence(Layers);
+	}
+
+	FModelBuilderElement FModelBuilder::MakeResidualMLPWithLayerNorm(
+		const uint32 InputSize,
+		const uint32 OutputSize,
+		const uint32 HiddenSize,
+		const uint32 LayerNum,
+		const EActivationFunction ActivationFunction,
+		const bool bActivationOnFinalLayer,
+		const FLinearLayerSettings& LinearLayerSettings)
+	{
+		check(LayerNum >= 2);
+
+		const int32 TotalLayerNum = LayerNum;
+
+		TArray<FModelBuilderElement, TInlineAllocator<32>> Layers;
+		Layers.Reserve(TotalLayerNum);
+
+		for (uint32 LayerIdx = 0; LayerIdx < LayerNum; LayerIdx++)
+		{
+			if (LayerIdx == 0)
+			{
+				Layers.Emplace(
+					MakeSequence({
+						MakeLinearLayer(InputSize, HiddenSize, LinearLayerSettings),
+						MakeActivation(HiddenSize, ActivationFunction),
+						}));
+			}
+			else if (LayerIdx != LayerNum - 1)
+			{
+				Layers.Emplace(MakeResidual(
+					MakeSequence({
+						MakeLayerNorm(HiddenSize, MakeValuesZero(HiddenSize), MakeValuesOne(HiddenSize)),
+						MakeLinearLayer(HiddenSize, HiddenSize, LinearLayerSettings),
+						MakeActivation(HiddenSize, ActivationFunction),
+						})));
+			}
+			else
+			{
+				if (bActivationOnFinalLayer)
+				{
+					Layers.Emplace(MakeSequence({
+							MakeLayerNorm(HiddenSize, MakeValuesZero(HiddenSize), MakeValuesOne(HiddenSize)),
+							MakeLinearLayer(HiddenSize, OutputSize, LinearLayerSettings),
+							MakeActivation(OutputSize, ActivationFunction),
+						}));
+				}
+				else
+				{
+					Layers.Emplace(MakeSequence({
+						MakeLayerNorm(HiddenSize, MakeValuesZero(HiddenSize), MakeValuesOne(HiddenSize)),
+						MakeLinearLayer(HiddenSize, OutputSize, LinearLayerSettings)
+					}));
+				}
+			}
+		}
+
+		check(Layers.Num() == TotalLayerNum);
+
+		return MakeSequence(Layers);
+	}
+
 	FModelBuilderElement FModelBuilder::MakeMemoryCell(
 		const uint32 InputNum,
 		const uint32 OutputNum,
@@ -5293,6 +5510,17 @@ namespace UE::NNE::RuntimeBasic
 		return StaticCastSharedPtr<Private::ILayer>(ArrayLayer);
 	}
 
+	FModelBuilderElement FModelBuilder::MakeResidual(const FModelBuilderElement& SubLayer)
+	{
+		check(SubLayer.GetInputSize() == SubLayer.GetOutputSize())
+
+		const TSharedPtr<Private::FResidualLayer> ResidualLayer = MakeShared<Private::FResidualLayer>();
+		ResidualLayer->InputOutputSize = SubLayer.GetInputSize();
+		ResidualLayer->SubLayer = SubLayer.Layer;
+
+		return StaticCastSharedPtr<Private::ILayer>(ResidualLayer);
+	}
+
 	FModelBuilderElement FModelBuilder::MakeAggregateSet(
 		const uint32 MaxElementNum,
 		const uint32 OutputEncodingSize,
diff --git a/Engine/Plugins/Experimental/NNERuntimeBasicCpu/Source/NNERuntimeBasicCpu/Public/NNERuntimeBasicCpuBuilder.h b/Engine/Plugins/Experimental/NNERuntimeBasicCpu/Source/NNERuntimeBasicCpu/Public/NNERuntimeBasicCpuBuilder.h
--- a/Engine/Plugins/Experimental/NNERuntimeBasicCpu/Source/NNERuntimeBasicCpu/Public/NNERuntimeBasicCpuBuilder.h
+++ b/Engine/Plugins/Experimental/NNERuntimeBasicCpu/Source/NNERuntimeBasicCpu/Public/NNERuntimeBasicCpuBuilder.h
@@ -325,6 +325,46 @@ namespace UE::NNE::RuntimeBasic
 			const EActivationFunction ActivationFunction,
 			const bool bActivationOnFinalLayer = false,
 			const FLinearLayerSettings& LinearLayerSettings = FLinearLayerSettings());
+		
+		/**
+		 * Makes a Multi-Layer Perceptron network using Residual layers.
+		 *
+		 * @param InputSize					Input Vector Size
+		 * @param OutputSize				Output Vector Size
+		 * @param HiddenSize				Number of hidden units to use on internal layers.
+		 * @param LayerNum					Number of layers. Includes input and output layers.
+		 * @param ActivationFunction		Activation function to use.
+		 * @param bActivationOnFinalLayer	If the activation function should be used on the final output layer.
+		 * @param LinearLayerSettings		Settings for the linear layers in the MLP.
+		 */
+		FModelBuilderElement MakeResidualMLP(
+			const uint32 InputSize,
+			const uint32 OutputSize,
+			const uint32 HiddenSize,
+			const uint32 LayerNum,
+			const EActivationFunction ActivationFunction,
+			const bool bActivationOnFinalLayer = false,
+			const FLinearLayerSettings& LinearLayerSettings = FLinearLayerSettings());
+
+		/**
+		 * Makes a Multi-Layer Perceptron network using Residual layers followed by LayerNorm layers.
+		 *
+		 * @param InputSize					Input Vector Size
+		 * @param OutputSize				Output Vector Size
+		 * @param HiddenSize				Number of hidden units to use on internal layers.
+		 * @param LayerNum					Number of layers. Includes input and output layers.
+		 * @param ActivationFunction		Activation function to use.
+		 * @param bActivationOnFinalLayer	If the activation function should be used on the final output layer.
+		 * @param LinearLayerSettings		Settings for the linear layers in the MLP.
+		 */
+		FModelBuilderElement MakeResidualMLPWithLayerNorm(
+			const uint32 InputSize,
+			const uint32 OutputSize,
+			const uint32 HiddenSize,
+			const uint32 LayerNum,
+			const EActivationFunction ActivationFunction,
+			const bool bActivationOnFinalLayer = false,
+			const FLinearLayerSettings& LinearLayerSettings = FLinearLayerSettings());
 
 		/**
 		 * Make a new Memory Cell layer.
@@ -396,6 +436,13 @@ namespace UE::NNE::RuntimeBasic
 		 */
 		FModelBuilderElement MakeArray(const uint32 ElementNum, const FModelBuilderElement& SubLayer);
 
+		/**
+		 * Make a residual layer which adds the result of evaluating the given sublayer on the input, to the input.
+		 *
+		 * @param SubLayer		Sublayer to evaluate on the input.
+		 */
+		FModelBuilderElement MakeResidual(const FModelBuilderElement& SubLayer);
+
 		/**
 		 * Make a layer which aggregates a set of other observations using attention. This is used by LearningAgents.
 		 */
