NNERuntimeBasicCpu: Added FiLM Layer

diff --git a/Engine/Plugins/Experimental/NNERuntimeBasicCpu/Source/NNERuntimeBasicCpu/Private/NNERuntimeBasicCpu.ispc b/Engine/Plugins/Experimental/NNERuntimeBasicCpu/Source/NNERuntimeBasicCpu/Private/NNERuntimeBasicCpu.ispc
--- a/Engine/Plugins/Experimental/NNERuntimeBasicCpu/Source/NNERuntimeBasicCpu/Private/NNERuntimeBasicCpu.ispc
+++ b/Engine/Plugins/Experimental/NNERuntimeBasicCpu/Source/NNERuntimeBasicCpu/Private/NNERuntimeBasicCpu.ispc
@@ -624,4 +624,21 @@ export void NNERuntimeBasicCPUOperatorLayerNorm(
 			Output[BatchIdx * OutputStride + Idx] = ((Input[BatchIdx * InputStride + Idx] - Mean) / Std) * Scale[Idx] + Offset[Idx];
 		}
 	}
+}
+
+export void NNERuntimeBasicCPUOperatorLayerFiLM(
+	uniform float InputOutput[],
+	uniform const float Condition[],
+	uniform const int BatchNum,
+	uniform const int InputOutputNum,
+	uniform const int InputOutputStride,
+	uniform const int ConditionStride)
+{
+	foreach (BatchIdx = 0 ... BatchNum, Idx = 0 ... InputOutputNum)
+	{
+		InputOutput[BatchIdx * InputOutputStride + Idx] =
+			(InputOutput[BatchIdx * InputOutputStride + Idx] *
+				Condition[BatchIdx * ConditionStride + Idx]) +
+			Condition[BatchIdx * ConditionStride + InputOutputNum + Idx];
+	}
 }
\ No newline at end of file
diff --git a/Engine/Plugins/Experimental/NNERuntimeBasicCpu/Source/NNERuntimeBasicCpu/Private/NNERuntimeBasicCpuModel.cpp b/Engine/Plugins/Experimental/NNERuntimeBasicCpu/Source/NNERuntimeBasicCpu/Private/NNERuntimeBasicCpuModel.cpp
--- a/Engine/Plugins/Experimental/NNERuntimeBasicCpu/Source/NNERuntimeBasicCpu/Private/NNERuntimeBasicCpuModel.cpp
+++ b/Engine/Plugins/Experimental/NNERuntimeBasicCpu/Source/NNERuntimeBasicCpu/Private/NNERuntimeBasicCpuModel.cpp
@@ -1496,6 +1496,40 @@ namespace UE::NNE::RuntimeBasic
 #endif
 		}
 
+		static inline void OperatorLayerFiLM(
+			float* RESTRICT InputOutput,
+			const float* RESTRICT Condition,
+			const uint32 BatchSize,
+			const uint32 InputOutputSize,
+			const uint32 InputOutputStride,
+			const uint32 ConditionStride)
+		{
+			NNE_RUNTIME_BASIC_TRACE_SCOPE(NNE::RuntimeBasic::Private::OperatorLayerFiLM);
+
+			NNE_RUNTIME_BASIC_CHECK_ALIASING(InputOutput, Condition);
+
+#if NNE_RUNTIME_BASIC_ENABLE_ISPC
+			ispc::NNERuntimeBasicCPUOperatorLayerFiLM(
+				InputOutput,
+				Condition,
+				BatchSize,
+				InputOutputSize,
+				InputOutputStride,
+				ConditionStride);
+#else
+			for (uint32 BatchIdx = 0; BatchIdx < BatchSize; BatchIdx++)
+			{
+				for (uint32 Idx = 0; Idx < InputOutputSize; Idx++)
+				{
+					InputOutput[BatchIdx * InputOutputStride + Idx] =
+						(InputOutput[BatchIdx * InputOutputStride + Idx] * 
+						Condition[BatchIdx * ConditionStride + Idx]) + 
+						Condition[BatchIdx * ConditionStride + InputOutputSize + Idx];
+				}
+			}
+#endif
+		}
+
 		//--------------------------------------------------------------------------
 		// Layer Types
 		//--------------------------------------------------------------------------
@@ -1530,6 +1564,7 @@ namespace UE::NNE::RuntimeBasic
 			Spread = 24,
 			Slice = 25,
 			Residual = 26,
+			FiLM = 27,
 		};
 
 		//--------------------------------------------------------------------------
@@ -4506,6 +4541,165 @@ namespace UE::NNE::RuntimeBasic
 			TSharedPtr<ILayer> SubLayer;
 		};
 
+
+		//--------------------------------------------------------------------------
+
+		struct FFiLMLayer;
+
+		struct FFiLMInstance : public ILayerInstance
+		{
+			FFiLMInstance(const FFiLMLayer& InFiLMLayer);
+
+			virtual void SetMaxBatchSize(const uint32 MaxBatchSize) override final;
+
+			const FFiLMLayer& FiLMLayer;
+			TSharedPtr<ILayerInstance> PrefixInstance;
+			TSharedPtr<ILayerInstance> ConditionInstance;
+			TSharedPtr<ILayerInstance> PostfixInstance;
+			TArray<float> PrefixBuffer;
+			TArray<float> ConditionBuffer;
+		};
+
+		struct FFiLMLayer : public ILayer
+		{
+			virtual TSharedPtr<ILayerInstance> MakeInstance() const { return MakeShared<FFiLMInstance>(*this); };
+			virtual ELayerType GetLayerType() const override final { return ELayerType::FiLM; }
+			virtual uint32 GetInputSize() const override final { return PrefixInputSize + ConditionInputSize; }
+			virtual uint32 GetOutputSize() const override final { return PostfixOutputSize; }
+
+			virtual void SerializationSize(uint64& InOutOffset) const override final
+			{
+				Serialization::Size(InOutOffset, PrefixInputSize);
+				Serialization::Size(InOutOffset, PrefixOutputSize);
+				Serialization::Size(InOutOffset, ConditionInputSize);
+				Serialization::Size(InOutOffset, ConditionOutputSize);
+				Serialization::Size(InOutOffset, PostfixInputSize);
+				Serialization::Size(InOutOffset, PostfixOutputSize);
+				Serialization::Size(InOutOffset, PrefixLayer);
+				Serialization::Size(InOutOffset, ConditionLayer);
+				Serialization::Size(InOutOffset, PostfixLayer);
+			}
+
+			virtual void SerializationLoad(uint64& InOutOffset, TConstArrayView<uint8> Data) override final
+			{
+				Serialization::Load(InOutOffset, PrefixInputSize, Data);
+				Serialization::Load(InOutOffset, PrefixOutputSize, Data);
+				Serialization::Load(InOutOffset, ConditionInputSize, Data);
+				Serialization::Load(InOutOffset, ConditionOutputSize, Data);
+				Serialization::Load(InOutOffset, PostfixInputSize, Data);
+				Serialization::Load(InOutOffset, PostfixOutputSize, Data);
+				Serialization::Load(InOutOffset, PrefixLayer, Data);
+				Serialization::Load(InOutOffset, ConditionLayer, Data);
+				Serialization::Load(InOutOffset, PostfixLayer, Data);
+			}
+
+			virtual void SerializationSave(uint64& InOutOffset, TArrayView<uint8> Data) const override final
+			{
+				Serialization::Save(InOutOffset, PrefixInputSize, Data);
+				Serialization::Save(InOutOffset, PrefixOutputSize, Data);
+				Serialization::Save(InOutOffset, ConditionInputSize, Data);
+				Serialization::Save(InOutOffset, ConditionOutputSize, Data);
+				Serialization::Save(InOutOffset, PostfixInputSize, Data);
+				Serialization::Save(InOutOffset, PostfixOutputSize, Data);
+				Serialization::Save(InOutOffset, PrefixLayer, Data);
+				Serialization::Save(InOutOffset, ConditionLayer, Data);
+				Serialization::Save(InOutOffset, PostfixLayer, Data);
+			}
+
+			virtual void Evaluate(
+				ILayerInstance* Instance,
+				float* OutputBuffer,
+				const float* InputBuffer,
+				const uint32 BatchSize,
+				const uint32 OutputBufferSize,
+				const uint32 InputBufferSize,
+				const uint32 OutputBufferStride,
+				const uint32 InputBufferStride) override final
+			{
+				NNE_RUNTIME_BASIC_TRACE_SCOPE(NNE::RuntimeBasic::Private::FFiLMLayer::Evaluate);
+				check(OutputBufferSize == GetOutputSize() && InputBufferSize == GetInputSize());
+				check(OutputBufferStride >= GetOutputSize() && InputBufferStride >= GetInputSize());
+				check(PrefixOutputSize * 2 == ConditionOutputSize);
+				check(PostfixInputSize == PrefixOutputSize);
+
+				FFiLMInstance* FiLMInstance = StaticCast<FFiLMInstance*>(Instance);
+				check(FiLMInstance);
+
+				OperatorNanCheck(InputBuffer, BatchSize, InputBufferSize, InputBufferStride);
+
+				PrefixLayer->Evaluate(
+					FiLMInstance->PrefixInstance.Get(),
+					FiLMInstance->PrefixBuffer.GetData(),
+					InputBuffer,
+					BatchSize,
+					PrefixOutputSize,
+					PrefixInputSize,
+					PrefixOutputSize,
+					InputBufferStride);
+
+				ConditionLayer->Evaluate(
+					FiLMInstance->ConditionInstance.Get(),
+					FiLMInstance->ConditionBuffer.GetData(),
+					InputBuffer + PrefixInputSize,
+					BatchSize,
+					ConditionOutputSize,
+					ConditionInputSize,
+					ConditionOutputSize,
+					InputBufferStride);
+
+				OperatorLayerFiLM(
+					FiLMInstance->PrefixBuffer.GetData(),
+					FiLMInstance->ConditionBuffer.GetData(),
+					BatchSize,
+					PrefixOutputSize,
+					PrefixOutputSize,
+					ConditionOutputSize);
+
+				PostfixLayer->Evaluate(
+					FiLMInstance->PostfixInstance.Get(),
+					OutputBuffer,
+					FiLMInstance->PrefixBuffer.GetData(),
+					BatchSize,
+					OutputBufferSize,
+					PostfixInputSize,
+					OutputBufferStride,
+					PostfixInputSize);
+
+				OperatorNanCheck(OutputBuffer, BatchSize, OutputBufferSize, OutputBufferStride);
+			}
+
+			uint32 PrefixInputSize = 0;
+			uint32 PrefixOutputSize = 0;
+			uint32 ConditionInputSize = 0;
+			uint32 ConditionOutputSize = 0;
+			uint32 PostfixInputSize = 0;
+			uint32 PostfixOutputSize = 0;
+
+			TSharedPtr<ILayer> PrefixLayer;
+			TSharedPtr<ILayer> ConditionLayer;
+			TSharedPtr<ILayer> PostfixLayer;
+		};
+
+		FFiLMInstance::FFiLMInstance(const FFiLMLayer& InFiLMLayer)
+			: FiLMLayer(InFiLMLayer)
+		{
+			PrefixInstance = InFiLMLayer.PrefixLayer->MakeInstance();
+			ConditionInstance = InFiLMLayer.ConditionLayer->MakeInstance();
+			PostfixInstance = InFiLMLayer.PostfixLayer->MakeInstance();
+		}
+
+		void FFiLMInstance::SetMaxBatchSize(const uint32 MaxBatchSize)
+		{
+			NNE_RUNTIME_BASIC_TRACE_SCOPE(NNE::RuntimeBasic::Private::FFiLMInstance::SetMaxBatchSize);
+
+			if (PrefixInstance) { PrefixInstance->SetMaxBatchSize(MaxBatchSize); }
+			if (ConditionInstance) { ConditionInstance->SetMaxBatchSize(MaxBatchSize); }
+			if (PostfixInstance) { PostfixInstance->SetMaxBatchSize(MaxBatchSize); }
+
+			PrefixBuffer.SetNumUninitialized(MaxBatchSize * FiLMLayer.PrefixOutputSize, EAllowShrinking::No);
+			ConditionBuffer.SetNumUninitialized(MaxBatchSize * FiLMLayer.ConditionOutputSize, EAllowShrinking::No);
+		}
+
 		//--------------------------------------------------------------------------
 		// Layer Serialization
 		//--------------------------------------------------------------------------
@@ -4560,6 +4754,7 @@ namespace UE::NNE::RuntimeBasic
 					case ELayerType::Spread: OutLayer = MakeShared<FSpreadLayer>(); break;
 					case ELayerType::Slice: OutLayer = MakeShared<FSliceLayer>(); break;
 					case ELayerType::Residual: OutLayer = MakeShared<FResidualLayer>(); break;
+					case ELayerType::FiLM: OutLayer = MakeShared<FFiLMLayer>(); break;
 					default: checkf(false, TEXT("Unknown Layer Id %i"), LayerTypeId);
 					}
 				}
@@ -5689,6 +5884,27 @@ namespace UE::NNE::RuntimeBasic
 		return StaticCastSharedPtr<Private::ILayer>(Layer);
 	}
 
+	FModelBuilderElement FModelBuilder::MakeFiLMNetwork(
+		const FModelBuilderElement& Prefix,
+		const FModelBuilderElement& Condition,
+		const FModelBuilderElement& Postfix)
+	{
+		check(Prefix.GetOutputSize() * 2 == Condition.GetOutputSize());
+		check(Prefix.GetOutputSize() == Postfix.GetInputSize());
+
+		const TSharedPtr<Private::FFiLMLayer> Layer = MakeShared<Private::FFiLMLayer>();
+		Layer->PrefixInputSize = Prefix.GetInputSize();
+		Layer->PrefixOutputSize = Prefix.GetOutputSize();
+		Layer->ConditionInputSize = Condition.GetInputSize();
+		Layer->ConditionOutputSize = Condition.GetOutputSize();
+		Layer->PostfixInputSize = Postfix.GetInputSize();
+		Layer->PostfixOutputSize = Postfix.GetOutputSize();
+		Layer->PrefixLayer = Prefix.Layer;
+		Layer->ConditionLayer = Condition.Layer;
+		Layer->PostfixLayer = Postfix.Layer;
+		return StaticCastSharedPtr<Private::ILayer>(Layer);
+	}
+
 	void FModelBuilder::Reset()
 	{
 		Rng = RngInitialState;
diff --git a/Engine/Plugins/Experimental/NNERuntimeBasicCpu/Source/NNERuntimeBasicCpu/Public/NNERuntimeBasicCpuBuilder.h b/Engine/Plugins/Experimental/NNERuntimeBasicCpu/Source/NNERuntimeBasicCpu/Public/NNERuntimeBasicCpuBuilder.h
--- a/Engine/Plugins/Experimental/NNERuntimeBasicCpu/Source/NNERuntimeBasicCpu/Public/NNERuntimeBasicCpuBuilder.h
+++ b/Engine/Plugins/Experimental/NNERuntimeBasicCpu/Source/NNERuntimeBasicCpu/Public/NNERuntimeBasicCpuBuilder.h
@@ -507,6 +507,17 @@ namespace UE::NNE::RuntimeBasic
 		/** Makes a Tile Layer */
 		FModelBuilderElement MakeTile(const uint32 InputSize, const uint32 Repeats);
 
+		/**
+		 * Make a new FiLM conditioned Network.
+		 *
+		 * @param Prefix				Prefix Element
+		 * @param Condition				Condition Element
+		 * @param Postfix				Postfix Element
+		 */
+		FModelBuilderElement MakeFiLMNetwork(
+			const FModelBuilderElement& Prefix,
+			const FModelBuilderElement& Condition,
+			const FModelBuilderElement& Postfix);
 
 	public:
 
