Moved dll exports to the right place (not on the type).

diff --git a/Engine/Plugins/Animation/PoseSearch/Source/Runtime/Public/PoseSearch/PoseSearchAnimNotifies.h b/Engine/Plugins/Animation/PoseSearch/Source/Runtime/Public/PoseSearch/PoseSearchAnimNotifies.h
--- a/Engine/Plugins/Animation/PoseSearch/Source/Runtime/Public/PoseSearch/PoseSearchAnimNotifies.h
+++ b/Engine/Plugins/Animation/PoseSearch/Source/Runtime/Public/PoseSearch/PoseSearchAnimNotifies.h
@@ -8,36 +8,38 @@
 #include "GameplayTagContainer.h"
 #include "PoseSearchAnimNotifies.generated.h"
 
+#define UE_API POSESEARCH_API
+
 class UPoseSearchDatabase;
 
 // Base class for ALL pose search anim notify states
-UCLASS(Abstract)
-class POSESEARCH_API UAnimNotifyState_PoseSearchBase : public UAnimNotifyState
+UCLASS(MinimalAPI, Abstract)
+class UAnimNotifyState_PoseSearchBase : public UAnimNotifyState
 {
 	GENERATED_BODY()
 };
 
 // Use this notify state to remove animation segments from the database completely, they will never play or return from
 // a search result
-UCLASS(Blueprintable, meta = (DisplayName = "Pose Search: Exclude From Database"))
-class POSESEARCH_API UAnimNotifyState_PoseSearchExcludeFromDatabase : public UAnimNotifyState_PoseSearchBase
+UCLASS(MinimalAPI, Blueprintable, meta = (DisplayName = "Pose Search: Exclude From Database"))
+class UAnimNotifyState_PoseSearchExcludeFromDatabase : public UAnimNotifyState_PoseSearchBase
 {
 	GENERATED_BODY()
 };
 
 // A pose search search will not return results that overlap this notify, but the animation segment can still play
 // if a previous search result advances into it.
-UCLASS(Blueprintable, meta = (DisplayName = "Pose Search: Block Transition In"))
-class POSESEARCH_API UAnimNotifyState_PoseSearchBlockTransition : public UAnimNotifyState_PoseSearchBase
+UCLASS(MinimalAPI, Blueprintable, meta = (DisplayName = "Pose Search: Block Transition In"))
+class UAnimNotifyState_PoseSearchBlockTransition : public UAnimNotifyState_PoseSearchBase
 {
 	GENERATED_BODY()
 };
 
 // Pose search cost will be affected by this, making the animation segment more or less likely to be selected based
 // on the notify parameters
-UCLASS(Blueprintable, meta = (DisplayName = "Pose Search: Override Base Cost Bias"))
+UCLASS(MinimalAPI, Blueprintable, meta = (DisplayName = "Pose Search: Override Base Cost Bias"))
 // @todo: rename into UAnimNotifyState_PoseSearchOverrideBaseCostBias
-class POSESEARCH_API UAnimNotifyState_PoseSearchModifyCost : public UAnimNotifyState_PoseSearchBase
+class UAnimNotifyState_PoseSearchModifyCost : public UAnimNotifyState_PoseSearchBase
 {
 	GENERATED_BODY()
 
@@ -51,8 +53,8 @@ public:
 
 // Pose search cost for the continuing pose will be affected by this, making the animation segment more or less 
 // likely to be continuing playing based on the notify parameters
-UCLASS(Blueprintable, meta = (DisplayName = "Pose Search: Override Continuing Pose Cost Bias"))
-class POSESEARCH_API UAnimNotifyState_PoseSearchOverrideContinuingPoseCostBias : public UAnimNotifyState_PoseSearchBase
+UCLASS(MinimalAPI, Blueprintable, meta = (DisplayName = "Pose Search: Override Continuing Pose Cost Bias"))
+class UAnimNotifyState_PoseSearchOverrideContinuingPoseCostBias : public UAnimNotifyState_PoseSearchBase
 {
 	GENERATED_BODY()
 
@@ -66,8 +68,8 @@ public:
 
 // UPoseSearchFeatureChannel(s) can use this UAnimNotifyState_PoseSearchSamplingEvent to demarcate events identified by SamplingAttributeId
 // during database indexing by specifying their SamplingAttributeId property to match UAnimNotifyState_PoseSearchSamplingAttribute::SamplingAttributeId
-UCLASS(Blueprintable, meta = (DisplayName = "Pose Search: Sampling Event"))
-class POSESEARCH_API UAnimNotifyState_PoseSearchSamplingEvent : public UAnimNotifyState_PoseSearchBase
+UCLASS(MinimalAPI, Blueprintable, meta = (DisplayName = "Pose Search: Sampling Event"))
+class UAnimNotifyState_PoseSearchSamplingEvent : public UAnimNotifyState_PoseSearchBase
 {
 	GENERATED_BODY()
 
@@ -83,8 +85,8 @@ public:
 
 // UPoseSearchFeatureChannel(s) can use this UAnimNotifyState_PoseSearchSamplingAttribute as animation space position, rotation, and linear velocity provider 
 // during database indexing by specifying their SamplingAttributeId property to match UAnimNotifyState_PoseSearchSamplingAttribute::SamplingAttributeId
-UCLASS(Blueprintable, meta = (DisplayName = "Pose Search: Sampling Attribute"))
-class POSESEARCH_API UAnimNotifyState_PoseSearchSamplingAttribute : public UAnimNotifyState_PoseSearchSamplingEvent
+UCLASS(MinimalAPI, Blueprintable, meta = (DisplayName = "Pose Search: Sampling Attribute"))
+class UAnimNotifyState_PoseSearchSamplingAttribute : public UAnimNotifyState_PoseSearchSamplingEvent
 {
 	GENERATED_BODY()
 
@@ -107,8 +109,8 @@ public:
 #endif // WITH_EDITORONLY_DATA
 };
 
-UCLASS(Blueprintable, meta = (DisplayName = "Pose Search: Motion Matched Branch In"))
-class POSESEARCH_API UAnimNotifyState_PoseSearchBranchIn : public UAnimNotifyState_PoseSearchBase
+UCLASS(MinimalAPI, Blueprintable, meta = (DisplayName = "Pose Search: Motion Matched Branch In"))
+class UAnimNotifyState_PoseSearchBranchIn : public UAnimNotifyState_PoseSearchBase
 {
 	GENERATED_BODY()
 
@@ -117,12 +119,12 @@ public:
 	UPROPERTY(EditAnywhere, Category = Config)
 	TObjectPtr<UPoseSearchDatabase> Database;
 
-	uint32 GetBranchInId() const;
+	UE_API uint32 GetBranchInId() const;
 };
 
 // multi character interaction notifies
 USTRUCT(Experimental, Blueprintable)
-struct POSESEARCH_API FPoseSearchIKWindowConstraint
+struct FPoseSearchIKWindowConstraint
 {
 	GENERATED_BODY()
 	
@@ -151,8 +153,8 @@ struct POSESEARCH_API FPoseSearchIKWindowConstraint
 	float DeactivationTime = 0.f;
 };
 
-UCLASS(Experimental, Blueprintable, meta = (DisplayName = "Pose Search: IKWindow"))
-class POSESEARCH_API UAnimNotifyState_PoseSearchIKWindow : public UAnimNotifyState_PoseSearchBase
+UCLASS(MinimalAPI, Experimental, Blueprintable, meta = (DisplayName = "Pose Search: IKWindow"))
+class UAnimNotifyState_PoseSearchIKWindow : public UAnimNotifyState_PoseSearchBase
 {
 	GENERATED_BODY()
 
@@ -162,14 +164,14 @@ public:
 };
 
 // Base class for ALL pose search anim notifies
-UCLASS(Abstract)
-class POSESEARCH_API UAnimNotify_PoseSearchBase : public UAnimNotify
+UCLASS(MinimalAPI, Abstract)
+class UAnimNotify_PoseSearchBase : public UAnimNotify
 {
 	GENERATED_BODY()
 };
 
-UCLASS(Experimental, Blueprintable, meta = (DisplayName = "Pose Search: Event"))
-class POSESEARCH_API UAnimNotify_PoseSearchEvent : public UAnimNotify_PoseSearchBase
+UCLASS(MinimalAPI, Experimental, Blueprintable, meta = (DisplayName = "Pose Search: Event"))
+class UAnimNotify_PoseSearchEvent : public UAnimNotify_PoseSearchBase
 {
 	GENERATED_BODY()
 
@@ -177,4 +179,6 @@ public:
 	// Tag identifying this event
 	UPROPERTY(EditAnywhere, Category = Config)
 	FGameplayTag EventTag;
-};
\ No newline at end of file
+};
+
+#undef UE_API
diff --git a/Engine/Plugins/Experimental/NNERuntimeBasicCpu/Source/NNERuntimeBasicCpu/Public/NNERuntimeBasicCpuBuilder.h b/Engine/Plugins/Experimental/NNERuntimeBasicCpu/Source/NNERuntimeBasicCpu/Public/NNERuntimeBasicCpuBuilder.h
--- a/Engine/Plugins/Experimental/NNERuntimeBasicCpu/Source/NNERuntimeBasicCpu/Public/NNERuntimeBasicCpuBuilder.h
+++ b/Engine/Plugins/Experimental/NNERuntimeBasicCpu/Source/NNERuntimeBasicCpu/Public/NNERuntimeBasicCpuBuilder.h
@@ -4,6 +4,8 @@
 
 #include <Math/RandomStream.h>
 
+#define UE_API NNERUNTIMEBASICCPU_API
+
 namespace UE::NNE::RuntimeBasic
 {
 	namespace Private
@@ -16,18 +18,18 @@ namespace UE::NNE::RuntimeBasic
 	/**
 	 * Represents an element of a model (such as a layer) created by the Model Builder.
 	 */
-	class NNERUNTIMEBASICCPU_API FModelBuilderElement
+	class FModelBuilderElement
 	{
 		friend class FModelBuilder;
 
 	public:
 
-		FModelBuilderElement();
-		FModelBuilderElement(const TSharedPtr<Private::ILayer>& Ptr);
-		~FModelBuilderElement();
+		UE_API FModelBuilderElement();
+		UE_API FModelBuilderElement(const TSharedPtr<Private::ILayer>& Ptr);
+		UE_API ~FModelBuilderElement();
 
-		int32 GetInputSize() const;
-		int32 GetOutputSize() const;
+		UE_API int32 GetInputSize() const;
+		UE_API int32 GetOutputSize() const;
 
 	private:
 		TSharedPtr<Private::ILayer> Layer;
@@ -42,7 +44,7 @@ namespace UE::NNE::RuntimeBasic
 	 * builder itself as it will not create internal copies. If you want to create an
 	 * internal copy use the `MakeValuesCopy` function.
 	 */
-	class NNERUNTIMEBASICCPU_API FModelBuilder
+	class FModelBuilder
 	{
 
 	public:
@@ -82,7 +84,7 @@ namespace UE::NNE::RuntimeBasic
 		 */
 
 		/** Weight initialization settings. */
-		struct NNERUNTIMEBASICCPU_API FWeightInitializationSettings
+		struct FWeightInitializationSettings
 		{
 			// Workaround for clang compiler bug...
 			FWeightInitializationSettings();
@@ -98,10 +100,10 @@ namespace UE::NNE::RuntimeBasic
 		};
 
 		/** Linear Layer settings including the type of layer and the way the weights are initialized */
-		struct NNERUNTIMEBASICCPU_API FLinearLayerSettings
+		struct FLinearLayerSettings
 		{
 			// Workaround for clang compiler bug...
-			FLinearLayerSettings();
+			UE_API FLinearLayerSettings();
 
 			// Type of the Linear Layer to construct
 			ELinearLayerType Type = ELinearLayerType::Normal;
@@ -111,7 +113,7 @@ namespace UE::NNE::RuntimeBasic
 		};
 
 		/** Construct a new Model Builder with the given random seed. */
-		FModelBuilder(int32 Seed = 0x0a974e75);
+		UE_API FModelBuilder(int32 Seed = 0x0a974e75);
 
 		/**
 		 * Makes a new linear layer.
@@ -121,7 +123,7 @@ namespace UE::NNE::RuntimeBasic
 		 * @param Weights		Linear layer weights.
 		 * @param Biases		Linear layer biases.
 		 */
-		FModelBuilderElement MakeLinear(
+		UE_API FModelBuilderElement MakeLinear(
 			const uint32 InputSize,
 			const uint32 OutputSize,
 			const TConstArrayView<float> Weights,
@@ -137,7 +139,7 @@ namespace UE::NNE::RuntimeBasic
 		 * @param WeightScales	Compressed Linear layer weight biases.
 		 * @param Biases		Compressed Linear layer biases.
 		 */
-		FModelBuilderElement MakeCompressedLinear(
+		UE_API FModelBuilderElement MakeCompressedLinear(
 			const uint32 InputSize,
 			const uint32 OutputSize,
 			const TConstArrayView<uint16> Weights,
@@ -153,7 +155,7 @@ namespace UE::NNE::RuntimeBasic
 		 * @param Weights		Linear layer weights.
 		 * @param Biases		Linear layer biases.
 		 */
-		FModelBuilderElement MakeLipschiztLinear(
+		UE_API FModelBuilderElement MakeLipschiztLinear(
 			const uint32 InputSize,
 			const uint32 OutputSize,
 			const TConstArrayView<float> Weights,
@@ -166,7 +168,7 @@ namespace UE::NNE::RuntimeBasic
 		 * @param OutputSize	Output Vector Size (Number of Columns)
 		 * @param Settings		Layer Settings.
 		 */
-		FModelBuilderElement MakeLinearLayer(
+		UE_API FModelBuilderElement MakeLinearLayer(
 			const uint32 InputSize,
 			const uint32 OutputSize,
 			const FLinearLayerSettings& Settings = FLinearLayerSettings());
@@ -180,7 +182,7 @@ namespace UE::NNE::RuntimeBasic
 		 * @param Weights		Multi-Linear layer weights.
 		 * @param Biases		Multi-Linear layer biases.
 		 */
-		FModelBuilderElement MakeMultiLinear(
+		UE_API FModelBuilderElement MakeMultiLinear(
 			const uint32 InputSize,
 			const uint32 OutputSize,
 			const uint32 BlockNum,
@@ -193,7 +195,7 @@ namespace UE::NNE::RuntimeBasic
 		 * @param Mean		Normalization Mean.
 		 * @param Std		Normalization Std.
 		 */
-		FModelBuilderElement MakeNormalize(
+		UE_API FModelBuilderElement MakeNormalize(
 			const uint32 InputOutputSize,
 			const TConstArrayView<float> Mean,
 			const TConstArrayView<float> Std);
@@ -204,34 +206,34 @@ namespace UE::NNE::RuntimeBasic
 		 * @param Mean		Denormalization Mean.
 		 * @param Std		Denormalization Std.
 		 */
-		FModelBuilderElement MakeDenormalize(
+		UE_API FModelBuilderElement MakeDenormalize(
 			const uint32 InputOutputSize,
 			const TConstArrayView<float> Mean,
 			const TConstArrayView<float> Std);
 
 		/** Makes a ReLU Activation Layer */
-		FModelBuilderElement MakeReLU(const uint32 InputOutputSize);
+		UE_API FModelBuilderElement MakeReLU(const uint32 InputOutputSize);
 
 		/** Makes a ELU Activation Layer */
-		FModelBuilderElement MakeELU(const uint32 InputOutputSize);
+		UE_API FModelBuilderElement MakeELU(const uint32 InputOutputSize);
 
 		/** Makes a GELU Activation Layer */
-		FModelBuilderElement MakeGELU(const uint32 InputOutputSize);
+		UE_API FModelBuilderElement MakeGELU(const uint32 InputOutputSize);
 
 		/** Makes a TanH Activation Layer */
-		FModelBuilderElement MakeTanH(const uint32 InputOutputSize);
+		UE_API FModelBuilderElement MakeTanH(const uint32 InputOutputSize);
 
 		/** Makes a Copy Layer */
-		FModelBuilderElement MakeCopy(const uint32 InputOutputSize);
+		UE_API FModelBuilderElement MakeCopy(const uint32 InputOutputSize);
 
 		/** Makes a Slice Layer */
-		FModelBuilderElement MakeSlice(const uint32 InputSize, const uint32 SliceOffset, const uint32 SliceSize);
+		UE_API FModelBuilderElement MakeSlice(const uint32 InputSize, const uint32 SliceOffset, const uint32 SliceSize);
 
 		/** Makes a Clamp Layer */
-		FModelBuilderElement MakeClamp(const uint32 InputOutputSize, const TConstArrayView<float> MinValues, const TConstArrayView<float> MaxValues);
+		UE_API FModelBuilderElement MakeClamp(const uint32 InputOutputSize, const TConstArrayView<float> MinValues, const TConstArrayView<float> MaxValues);
 
 		/** Makes a new activation layer with the given activation function */
-		FModelBuilderElement MakeActivation(const uint32 InputOutputSize, const EActivationFunction ActivationFunction);
+		UE_API FModelBuilderElement MakeActivation(const uint32 InputOutputSize, const EActivationFunction ActivationFunction);
 
 		/**
 		 * Makes a PReLU Activation Layer
@@ -239,12 +241,12 @@ namespace UE::NNE::RuntimeBasic
 		 * @param LayerSize			Number of neurons in layer.
 		 * @param Alpha				PReLU alpha parameter for each neuron.
 		 */
-		FModelBuilderElement MakePReLU(const uint32 InputOutputSize, const TConstArrayView<float> Alpha);
+		UE_API FModelBuilderElement MakePReLU(const uint32 InputOutputSize, const TConstArrayView<float> Alpha);
 
 		/**
 		 * Makes a Sequence layer, which will evaluate the given list of layers in order.
 		 */
-		FModelBuilderElement MakeSequence(const TConstArrayView<FModelBuilderElement> Elements);
+		UE_API FModelBuilderElement MakeSequence(const TConstArrayView<FModelBuilderElement> Elements);
 
 		/**
 		 * Makes a Multi-Layer Perceptron network.
@@ -257,7 +259,7 @@ namespace UE::NNE::RuntimeBasic
 		 * @param bActivationOnFinalLayer	If the activation function should be used on the final output layer.
 		 * @param LinearLayerSettings		Settings for the linear layers in the MLP.
 		 */
-		FModelBuilderElement MakeMLP(
+		UE_API FModelBuilderElement MakeMLP(
 			const uint32 InputSize,
 			const uint32 OutputSize, 
 			const uint32 HiddenSize,
@@ -277,7 +279,7 @@ namespace UE::NNE::RuntimeBasic
 		 * @param bActivationOnFinalLayer	If the activation function should be used on the final output layer.
 		 * @param LinearLayerSettings		Settings for the linear layers in the MLP.
 		 */
-		FModelBuilderElement MakeMLPWithLayerNorm(
+		UE_API FModelBuilderElement MakeMLPWithLayerNorm(
 			const uint32 InputSize,
 			const uint32 OutputSize,
 			const uint32 HiddenSize,
@@ -297,7 +299,7 @@ namespace UE::NNE::RuntimeBasic
 		 * @param bActivationOnFinalLayer	If the activation function should be used on the final output layer.
 		 * @param LinearLayerSettings		Settings for the linear layers in the MLP.
 		 */
-		FModelBuilderElement MakeSkipMLP(
+		UE_API FModelBuilderElement MakeSkipMLP(
 			const uint32 InputSize,
 			const uint32 OutputSize,
 			const uint32 HiddenSize,
@@ -317,7 +319,7 @@ namespace UE::NNE::RuntimeBasic
 		 * @param bActivationOnFinalLayer	If the activation function should be used on the final output layer.
 		 * @param LinearLayerSettings		Settings for the linear layers in the MLP.
 		 */
-		FModelBuilderElement MakeSkipMLPWithLayerNorm(
+		UE_API FModelBuilderElement MakeSkipMLPWithLayerNorm(
 			const uint32 InputSize,
 			const uint32 OutputSize,
 			const uint32 HiddenSize,
@@ -337,7 +339,7 @@ namespace UE::NNE::RuntimeBasic
 		 * @param bActivationOnFinalLayer	If the activation function should be used on the final output layer.
 		 * @param LinearLayerSettings		Settings for the linear layers in the MLP.
 		 */
-		FModelBuilderElement MakeResidualMLP(
+		UE_API FModelBuilderElement MakeResidualMLP(
 			const uint32 InputSize,
 			const uint32 OutputSize,
 			const uint32 HiddenSize,
@@ -357,7 +359,7 @@ namespace UE::NNE::RuntimeBasic
 		 * @param bActivationOnFinalLayer	If the activation function should be used on the final output layer.
 		 * @param LinearLayerSettings		Settings for the linear layers in the MLP.
 		 */
-		FModelBuilderElement MakeResidualMLPWithLayerNorm(
+		UE_API FModelBuilderElement MakeResidualMLPWithLayerNorm(
 			const uint32 InputSize,
 			const uint32 OutputSize,
 			const uint32 HiddenSize,
@@ -378,7 +380,7 @@ namespace UE::NNE::RuntimeBasic
 		 * @param OutputInputUpdateLayer	Layer used for updating the output from the input.
 		 * @param OutputMemoryUpdateLayer	Layer used for updating the output from the memory.
 		 */
-		FModelBuilderElement MakeMemoryCell(
+		UE_API FModelBuilderElement MakeMemoryCell(
 			const uint32 InputNum,
 			const uint32 OutputNum,
 			const uint32 MemoryNum,
@@ -396,7 +398,7 @@ namespace UE::NNE::RuntimeBasic
 		 * @param MemoryNum				The size of the memory vector used by the model
 		 * @param LinearLayerSettings	Settings for the linear layers in the Memory Cell.
 		 */
-		FModelBuilderElement MakeMemoryCellLayer(
+		UE_API FModelBuilderElement MakeMemoryCellLayer(
 			const uint32 InputNum,
 			const uint32 OutputNum,
 			const uint32 MemoryNum,
@@ -410,7 +412,7 @@ namespace UE::NNE::RuntimeBasic
 		 * @param Cell					Memory Cell Element
 		 * @param Postfix				Postfix Element
 		 */
-		FModelBuilderElement MakeMemoryBackbone(
+		UE_API FModelBuilderElement MakeMemoryBackbone(
 			const uint32 MemoryNum,
 			const FModelBuilderElement& Prefix,
 			const FModelBuilderElement& Cell,
@@ -420,13 +422,13 @@ namespace UE::NNE::RuntimeBasic
 		 * Makes a Concat layer, which will evaluate each of the given elements on different slices of 
 		 * the input vector, concatenating the result into the output vector.
 		 */
-		FModelBuilderElement MakeConcat(const TConstArrayView<FModelBuilderElement> Elements);
+		UE_API FModelBuilderElement MakeConcat(const TConstArrayView<FModelBuilderElement> Elements);
 		
 		/**
 		 * Makes a Spread layer, which will evaluate each of the given elements on the input vector, 
 		 * concatenating the result into the output vector.
 		 */
-		FModelBuilderElement MakeSpread(const TConstArrayView<FModelBuilderElement> Elements);
+		UE_API FModelBuilderElement MakeSpread(const TConstArrayView<FModelBuilderElement> Elements);
 
 		/**
 		 * Make a layer which runs the given sublayer on an array of elements.
@@ -434,19 +436,19 @@ namespace UE::NNE::RuntimeBasic
 		 * @param ElementNum	Number of elements in the array.
 		 * @param SubLayer		Sublayer to evaluate on each item in the array.
 		 */
-		FModelBuilderElement MakeArray(const uint32 ElementNum, const FModelBuilderElement& SubLayer);
+		UE_API FModelBuilderElement MakeArray(const uint32 ElementNum, const FModelBuilderElement& SubLayer);
 
 		/**
 		 * Make a residual layer which adds the result of evaluating the given sublayer on the input, to the input.
 		 *
 		 * @param SubLayer		Sublayer to evaluate on the input.
 		 */
-		FModelBuilderElement MakeResidual(const FModelBuilderElement& SubLayer);
+		UE_API FModelBuilderElement MakeResidual(const FModelBuilderElement& SubLayer);
 
 		/**
 		 * Make a layer which aggregates a set of other observations using attention. This is used by LearningAgents.
 		 */
-		FModelBuilderElement MakeAggregateSet(
+		UE_API FModelBuilderElement MakeAggregateSet(
 			const uint32 MaxElementNum, 
 			const uint32 OutputEncodingSize,
 			const uint32 AttentionEncodingSize,
@@ -459,7 +461,7 @@ namespace UE::NNE::RuntimeBasic
 		/**
 		 * Make a layer which aggregates an Exclusive Or of other observations. This is used by LearningAgents.
 		 */
-		FModelBuilderElement MakeAggregateOrExclusive(
+		UE_API FModelBuilderElement MakeAggregateOrExclusive(
 			const uint32 OutputEncodingSize,
 			const TConstArrayView<FModelBuilderElement> SubLayers,
 			const TConstArrayView<FModelBuilderElement> Encoders);
@@ -467,7 +469,7 @@ namespace UE::NNE::RuntimeBasic
 		/**
 		 * Make a layer which aggregates an Inclusive Or of other observations using attention. This is used by LearningAgents.
 		 */
-		FModelBuilderElement MakeAggregateOrInclusive(
+		UE_API FModelBuilderElement MakeAggregateOrInclusive(
 			const uint32 OutputEncodingSize,
 			const uint32 AttentionEncodingSize,
 			const uint32 AttentionHeadNum,
@@ -484,7 +486,7 @@ namespace UE::NNE::RuntimeBasic
 		 * @param GatingLayer				Layer used to choose the experts. Input should be of size InputNum, output should be of size SubLayerNum.
 		 * @param SubLayers					Expert layers. All Layers here should have input of size InputNum, output of size OutputNum.
 		 */
-		FModelBuilderElement MakeSparseMixtureOfExperts(
+		UE_API FModelBuilderElement MakeSparseMixtureOfExperts(
 			const uint32 InputNum,
 			const uint32 OutputNum,
 			const FModelBuilderElement& GatingLayer,
@@ -498,14 +500,14 @@ namespace UE::NNE::RuntimeBasic
 		 * @param Scales			Initial Scales
 		 * @param Epsilon			Standard Deviation Epsilon
 		 */
-		FModelBuilderElement MakeLayerNorm(
+		UE_API FModelBuilderElement MakeLayerNorm(
 			const uint32 InputOutputSize,
 			const TConstArrayView<float> Offsets,
 			const TConstArrayView<float> Scales,
 			const float Epsilon = 1e-5f);
 
 		/** Makes a Tile Layer */
-		FModelBuilderElement MakeTile(const uint32 InputSize, const uint32 Repeats);
+		UE_API FModelBuilderElement MakeTile(const uint32 InputSize, const uint32 Repeats);
 
 		/**
 		 * Make a new FiLM conditioned Network.
@@ -514,7 +516,7 @@ namespace UE::NNE::RuntimeBasic
 		 * @param Condition				Condition Element
 		 * @param Postfix				Postfix Element
 		 */
-		FModelBuilderElement MakeFiLMNetwork(
+		UE_API FModelBuilderElement MakeFiLMNetwork(
 			const FModelBuilderElement& Prefix,
 			const FModelBuilderElement& Condition,
 			const FModelBuilderElement& Postfix);
@@ -522,31 +524,31 @@ namespace UE::NNE::RuntimeBasic
 	public:
 
 		/** Creates a array of values from a copy of the given array view */
-		TArrayView<float> MakeValuesCopy(const TConstArrayView<float> Values);
+		UE_API TArrayView<float> MakeValuesCopy(const TConstArrayView<float> Values);
 
 		/** Creates a array of values set to zero of the given size */
-		TArrayView<float> MakeValuesZero(const uint32 Size);
+		UE_API TArrayView<float> MakeValuesZero(const uint32 Size);
 
 		/** Creates a array of values set to one of the given size */
-		TArrayView<float> MakeValuesOne(const uint32 Size);
+		UE_API TArrayView<float> MakeValuesOne(const uint32 Size);
 
 		/** Creates a array of values set to the provided constant value of the given size */
-		TArrayView<float> MakeValuesConstant(const uint32 Size, const float Value);
+		UE_API TArrayView<float> MakeValuesConstant(const uint32 Size, const float Value);
 
 		/** Creates a array of weights randomly initialized using the Gaussian Kaiming method */
-		TArrayView<float> MakeWeightsRandomKaimingGaussian(const uint32 InputSize, const uint32 OutputSize, const float Scale = 1.0f);
+		UE_API TArrayView<float> MakeWeightsRandomKaimingGaussian(const uint32 InputSize, const uint32 OutputSize, const float Scale = 1.0f);
 
 		/** Creates a array of weights randomly initialized using the Uniform Kaiming method */
-		TArrayView<float> MakeWeightsRandomKaimingUniform(const uint32 InputSize, const uint32 OutputSize, const float Scale = 1.0f);
+		UE_API TArrayView<float> MakeWeightsRandomKaimingUniform(const uint32 InputSize, const uint32 OutputSize, const float Scale = 1.0f);
 
 		/** Creates a array of biases randomly initialized using the Gaussian Kaiming method */
-		TArrayView<float> MakeBiasesRandomKaimingGaussian(const uint32 Size, const float Scale = 1.0f);
+		UE_API TArrayView<float> MakeBiasesRandomKaimingGaussian(const uint32 Size, const float Scale = 1.0f);
 
 		/** Creates a array of biases randomly initialized using the Uniform Kaiming method# */
-		TArrayView<float> MakeBiasesRandomKaimingUniform(const uint32 Size, const float Scale = 1.0f);
+		UE_API TArrayView<float> MakeBiasesRandomKaimingUniform(const uint32 Size, const float Scale = 1.0f);
 
 		/** Creates a array of weights randomly initialized using the Gaussian Kaiming method and compresses them */
-		void MakeCompressedWeightsRandomKaimingGaussian(
+		UE_API void MakeCompressedWeightsRandomKaimingGaussian(
 			TArrayView<uint16>& OutWeightsView,
 			TArrayView<float>& OutWeightOffsetsView,
 			TArrayView<float>& OutWeightScalesView,
@@ -555,7 +557,7 @@ namespace UE::NNE::RuntimeBasic
 			const float Scale = 1.0f);
 
 		/** Creates a array of weights randomly initialized using the Uniform Kaiming method and compresses them */
-		void MakeCompressedWeightsRandomKaimingUniform(
+		UE_API void MakeCompressedWeightsRandomKaimingUniform(
 			TArrayView<uint16>& OutWeightsView,
 			TArrayView<float>& OutWeightOffsetsView,
 			TArrayView<float>& OutWeightScalesView,
@@ -564,18 +566,18 @@ namespace UE::NNE::RuntimeBasic
 			const float Scale = 1.0f);
 
 		/** Creates a array of weights initialized using the given settings */
-		TArrayView<float> MakeInitialWeights(
+		UE_API TArrayView<float> MakeInitialWeights(
 			const uint32 InputSize, 
 			const uint32 OutputSize,
 			const FWeightInitializationSettings& WeightInitializationSettings = FWeightInitializationSettings());
 
 		/** Creates a array of biases initialized using the given settings */
-		TArrayView<float> MakeInitialBiases(
+		UE_API TArrayView<float> MakeInitialBiases(
 			const uint32 OutputSize,
 			const FWeightInitializationSettings& WeightInitializationSettings = FWeightInitializationSettings());
 
 		/** Creates a array of weights randomly initialized using the given settings and compresses them */
-		void MakeInitialCompressedWeights(
+		UE_API void MakeInitialCompressedWeights(
 			TArrayView<uint16>& OutWeightsView,
 			TArrayView<float>& OutWeightOffsetsView,
 			TArrayView<float>& OutWeightScalesView,
@@ -584,23 +586,23 @@ namespace UE::NNE::RuntimeBasic
 			const FWeightInitializationSettings& WeightInitializationSettings = FWeightInitializationSettings());
 
 		/** Creates an array of sizes, initialized to zero. */
-		TArrayView<uint32> MakeSizesZero(const uint32 Size);
+		UE_API TArrayView<uint32> MakeSizesZero(const uint32 Size);
 
 		/** Creates an array of sizes from an array of builder elements' input sizes. */
-		TArrayView<uint32> MakeSizesLayerInputs(const TConstArrayView<FModelBuilderElement> Elements);
+		UE_API TArrayView<uint32> MakeSizesLayerInputs(const TConstArrayView<FModelBuilderElement> Elements);
 
 		/** Creates an array of sizes from an array of builder elements' output sizes. */
-		TArrayView<uint32> MakeSizesLayerOutputs(const TConstArrayView<FModelBuilderElement> Elements);
+		UE_API TArrayView<uint32> MakeSizesLayerOutputs(const TConstArrayView<FModelBuilderElement> Elements);
 
 	public:
 
 		/** Reset the builder clearing all memory. */
-		void Reset();
+		UE_API void Reset();
 
 		/**
 		 * Get the number of bytes this builder currently wants to write for the given element.
 		 */
-		uint64 GetWriteByteNum(const FModelBuilderElement& Element) const;
+		UE_API uint64 GetWriteByteNum(const FModelBuilderElement& Element) const;
 
 		/**
 		 * Write the Model to FileData. Use `GetWriteByteNum` to get the number of bytes this will write so that 
@@ -611,7 +613,7 @@ namespace UE::NNE::RuntimeBasic
 		 * @param OutOutputSize			Output for the number of float outputs from the model
 		 * @param Element				Builder element to write the model FileData for.
 		 */
-		void WriteFileData(TArrayView<uint8> OutFileData, uint32& OutInputSize, uint32& OutOutputSize, const FModelBuilderElement& Element) const;
+		UE_API void WriteFileData(TArrayView<uint8> OutFileData, uint32& OutInputSize, uint32& OutOutputSize, const FModelBuilderElement& Element) const;
 
 		/**
 		 * Write the model to the FileData. Use `GetWriteByteNum` to get the number of bytes this will write so that 
@@ -622,7 +624,7 @@ namespace UE::NNE::RuntimeBasic
 		 * @param OutOutputSize			Output for the number of float outputs from the model
 		 * @param Element				Builder element to write the model FileData for.
 		 */
-		void WriteFileData(TArray<uint8>& OutFileData, uint32& OutInputSize, uint32& OutOutputSize, const FModelBuilderElement& Element) const;
+		UE_API void WriteFileData(TArray<uint8>& OutFileData, uint32& OutInputSize, uint32& OutOutputSize, const FModelBuilderElement& Element) const;
 
 		/**
 		 * Write the model to the FileData and Reset the builder clearing all memory used.
@@ -632,7 +634,7 @@ namespace UE::NNE::RuntimeBasic
 		 * @param OutOutputSize			Output for the number of float outputs from the model
 		 * @param Element				Builder element to write the model FileData for.
 		 */
-		void WriteFileDataAndReset(TArrayView<uint8> OutFileData, uint32& OutInputSize, uint32& OutOutputSize, const FModelBuilderElement& Element);
+		UE_API void WriteFileDataAndReset(TArrayView<uint8> OutFileData, uint32& OutInputSize, uint32& OutOutputSize, const FModelBuilderElement& Element);
 
 		/**
 		 * Write the model to the FileData and Reset the builder clearing all memory used.
@@ -642,7 +644,7 @@ namespace UE::NNE::RuntimeBasic
 		 * @param OutOutputSize			Output for the number of float outputs from the model
 		 * @param Element				Builder element to write the model FileData for.
 		 */
-		void WriteFileDataAndReset(TArray<uint8>& OutFileData, uint32& OutInputSize, uint32& OutOutputSize, const FModelBuilderElement& Element);
+		UE_API void WriteFileDataAndReset(TArray<uint8>& OutFileData, uint32& OutInputSize, uint32& OutOutputSize, const FModelBuilderElement& Element);
 
 	private:
 
@@ -662,4 +664,6 @@ namespace UE::NNE::RuntimeBasic
 		TArray<TArray<uint32>> SizesPool;
 	};
 
-}
\ No newline at end of file
+}
+
+#undef UE_API 
\ No newline at end of file
diff --git a/Engine/Plugins/NetcodeUnitTest/NetcodeUnitTest/Source/NetcodeUnitTest/Public/Net/NUTUtilNet.h b/Engine/Plugins/NetcodeUnitTest/NetcodeUnitTest/Source/NetcodeUnitTest/Public/Net/NUTUtilNet.h
--- a/Engine/Plugins/NetcodeUnitTest/NetcodeUnitTest/Source/NetcodeUnitTest/Public/Net/NUTUtilNet.h
+++ b/Engine/Plugins/NetcodeUnitTest/NetcodeUnitTest/Source/NetcodeUnitTest/Public/Net/NUTUtilNet.h
@@ -366,7 +366,7 @@ private:
  */
 class FNetPropertyHook : public FProperty
 {
-	DECLARE_FIELD(FNetPropertyHook, FProperty, CASTCLASS_None)
+	DECLARE_FIELD_API(FNetPropertyHook, FProperty, CASTCLASS_None, NETCODEUNITTEST_API)
 
 private:
 	virtual bool NetSerializeItem(FArchive& Ar, UPackageMap* Map, void* Data, TArray<uint8>* MetaData /*=nullptr*/) const override;
diff --git a/Engine/Plugins/Online/OnlineFramework/Source/Party/Public/Interactions/SocialInteractionMacros.h b/Engine/Plugins/Online/OnlineFramework/Source/Party/Public/Interactions/SocialInteractionMacros.h
--- a/Engine/Plugins/Online/OnlineFramework/Source/Party/Public/Interactions/SocialInteractionMacros.h
+++ b/Engine/Plugins/Online/OnlineFramework/Source/Party/Public/Interactions/SocialInteractionMacros.h
@@ -48,7 +48,7 @@ private:
 // Establishes boilerplate behavior and declares all functions the user is required to provide
 
 #define DECLARE_SOCIAL_INTERACTION_EXPORT(APIMacro, InteractionName)	\
-class APIMacro FSocialInteraction_##InteractionName	\
+class FSocialInteraction_##InteractionName	\
 {	\
 public:	\
 	static FSocialInteractionHandle GetHandle()	\
@@ -61,8 +61,8 @@ public:	\
 		return #InteractionName; \
 	} \
 	\
-	static FText GetDisplayName(const USocialUser& User);	\
-	static FString GetSlashCommandToken();	\
+	static APIMacro FText GetDisplayName(const USocialUser& User);	\
+	static APIMacro FString GetSlashCommandToken();	\
 	\
 	static bool IsAvailable(const USocialUser& User)	\
 	{	\
@@ -80,7 +80,7 @@ public:	\
 		});	\
 	}	\
 	\
-	static void ExecuteInteraction(USocialUser& User);	\
+	static APIMacro void ExecuteInteraction(USocialUser& User);	\
 	\
 	static FOnCustomIsInteractionAvailable& OnCustomIsInteractionAvailable()	\
 	{	\
@@ -88,7 +88,7 @@ public:	\
 		return CustomAvailabilityCheckDelegate;	\
 	}	\
 private: \
-	static bool CanExecute(const USocialUser& User); \
+	static APIMacro bool CanExecute(const USocialUser& User); \
 }
 
 #define DECLARE_SOCIAL_INTERACTION(InteractionName) DECLARE_SOCIAL_INTERACTION_EXPORT(, InteractionName)
diff --git a/Engine/Plugins/PCG/Source/PCG/Public/PCGSettings.h b/Engine/Plugins/PCG/Source/PCG/Public/PCGSettings.h
--- a/Engine/Plugins/PCG/Source/PCG/Public/PCGSettings.h
+++ b/Engine/Plugins/PCG/Source/PCG/Public/PCGSettings.h
@@ -14,6 +14,8 @@
 
 #include "PCGSettings.generated.h"
 
+#define UE_API PCG_API
+
 class UPCGComponent;
 class UPCGPin;
 class UPCGComputeKernel;
@@ -127,18 +129,18 @@ struct FPCGSettingsOverridableParam
 #endif // WITH_EDITOR
 };
 
-UCLASS(Abstract)
-class PCG_API UPCGSettingsInterface : public UPCGData
+UCLASS(MinimalAPI, Abstract)
+class UPCGSettingsInterface : public UPCGData
 {
 	GENERATED_BODY()
 
 public:
-	virtual UPCGSettings* GetSettings() PURE_VIRTUAL(UPCGSettingsInterface::GetSettings, return nullptr;);
-	virtual const UPCGSettings* GetSettings() const PURE_VIRTUAL(UPCGSettingsInterface::GetSettings, return nullptr;);
+	UE_API virtual UPCGSettings* GetSettings() PURE_VIRTUAL(UPCGSettingsInterface::GetSettings, return nullptr;);
+	UE_API virtual const UPCGSettings* GetSettings() const PURE_VIRTUAL(UPCGSettingsInterface::GetSettings, return nullptr;);
 
-	bool IsInstance() const;
+	UE_API bool IsInstance() const;
 	/** Dedicated method to change enable state because some nodes have more complex behavior on enable/disable (such as subgraphs) */
-	void SetEnabled(bool bInEnabled);
+	UE_API void SetEnabled(bool bInEnabled);
 
 	/** Whether this element can be disabled. */
 	virtual bool CanBeDisabled() const { return true; }
@@ -176,8 +178,8 @@ public:
 /**
 * Base class for settings-as-data in the PCG framework
 */
-UCLASS(Abstract, BlueprintType, ClassGroup = (Procedural))
-class PCG_API UPCGSettings : public UPCGSettingsInterface
+UCLASS(MinimalAPI, Abstract, BlueprintType, ClassGroup = (Procedural))
+class UPCGSettings : public UPCGSettingsInterface
 {
 	GENERATED_BODY()
 
@@ -188,7 +190,7 @@ class PCG_API UPCGSettings : public UPCGSettingsInterface
 public:
 	// ~Begin UPCGData interface
 	virtual EPCGDataType GetDataType() const override { return EPCGDataType::Settings; }
-	virtual void AddToCrc(FArchiveCrc32& Ar, bool bFullDataCrc) const override;
+	UE_API virtual void AddToCrc(FArchiveCrc32& Ar, bool bFullDataCrc) const override;
 	// ~End UPCGData interface
 
 	// ~Begin UPCGSettingsInterface interface
@@ -197,26 +199,26 @@ public:
 	// ~End UPCGSettingsInterface interface
 
 	//~Begin UObject interface
-	virtual void PostLoad() override;
-	virtual void Serialize(FArchive& Ar) override;
-	virtual void PostSaveRoot(FObjectPostSaveRootContext ObjectSaveContext) override;
-	virtual void PostInitProperties() override;
+	UE_API virtual void PostLoad() override;
+	UE_API virtual void Serialize(FArchive& Ar) override;
+	UE_API virtual void PostSaveRoot(FObjectPostSaveRootContext ObjectSaveContext) override;
+	UE_API virtual void PostInitProperties() override;
 #if WITH_EDITOR
-	virtual void PostEditUndo() override;
+	UE_API virtual void PostEditUndo() override;
 #endif // WITH_EDITOR
 	//~End UObject interface
 
-	void OnOverrideSettingsDuplicated(bool bSkippedPostLoad);
+	UE_API void OnOverrideSettingsDuplicated(bool bSkippedPostLoad);
 
 	// TODO: check if we need this to be virtual, we don't really need if we're always caching
-	/*virtual*/ FPCGElementPtr GetElement() const;
-	virtual UPCGNode* CreateNode() const;
+	UE_API /*virtual*/ FPCGElementPtr GetElement() const;
+	UE_API virtual UPCGNode* CreateNode() const;
 	
 	/** Return the concatenation of InputPinProperties, ExecutionDependencyPin, and FillOverridableParamsPins */
-	TArray<FPCGPinProperties> AllInputPinProperties() const;
+	UE_API TArray<FPCGPinProperties> AllInputPinProperties() const;
 
 	/** For symmetry reason, do the same with output pins. For now forward just the call to OutputPinProperties */
-	TArray<FPCGPinProperties> AllOutputPinProperties() const;
+	UE_API TArray<FPCGPinProperties> AllOutputPinProperties() const;
 
 	/** If the node has any dynamic pins that can change based on input or settings */
 	virtual bool HasDynamicPins() const { return false; }
@@ -225,28 +227,28 @@ public:
 	virtual bool OutputPinsCanBeDeactivated() const { return false; }
 
 	/** Bitwise union of the allowed types of each incident edge on pin. Returns None if no common bits, or no edges. */
-	EPCGDataType GetTypeUnionOfIncidentEdges(const FName& PinLabel) const;
+	UE_API EPCGDataType GetTypeUnionOfIncidentEdges(const FName& PinLabel) const;
 
 	// Internal functions, should not be used by any user.
 	// Return a different subset for input/output pin properties, in case of a default object.
-	virtual TArray<FPCGPinProperties> DefaultInputPinProperties() const;
-	virtual TArray<FPCGPinProperties> DefaultOutputPinProperties() const;
+	UE_API virtual TArray<FPCGPinProperties> DefaultInputPinProperties() const;
+	UE_API virtual TArray<FPCGPinProperties> DefaultOutputPinProperties() const;
 
-	bool operator==(const UPCGSettings& Other) const;
+	UE_API bool operator==(const UPCGSettings& Other) const;
 
 	PRAGMA_DISABLE_DEPRECATION_WARNINGS
 	/** By default, settings do not use a seed. Override this in the settings subclass to enable usage of the seed. UFUNCTION to be used by EditCondition. */
 	UFUNCTION()
 	virtual bool UseSeed() const { return bUseSeed; }
-	PRAGMA_ENABLE_DEPRECATION_WARNINGS
+	UE_API PRAGMA_ENABLE_DEPRECATION_WARNINGS
 
 	// Get the seed, combined with optional ExecutionSource seed
 	int GetSeed(const IPCGGraphExecutionSource* InExecutionSource = nullptr) const;
 
 	UE_DEPRECATED(5.6, "Use GetSettingsCrc instead")
-	const FPCGCrc& GetCachedCrc() const;
+	UE_API const FPCGCrc& GetCachedCrc() const;
 
-	FPCGCrc GetSettingsCrc() const;
+	UE_API FPCGCrc GetSettingsCrc() const;
 
 #if WITH_EDITOR
 	/** Puts node title on node body, reducing overall node size */
@@ -257,12 +259,12 @@ public:
 	virtual bool CanUserEditTitle() const { return true; }
 
 	/** UpdatePins will kick off invalid edges, so this is useful for moving edges around in case of pin changes. */
-	virtual void ApplyDeprecationBeforeUpdatePins(UPCGNode* InOutNode, TArray<TObjectPtr<UPCGPin>>& InputPins, TArray<TObjectPtr<UPCGPin>>& OutputPins);
+	UE_API virtual void ApplyDeprecationBeforeUpdatePins(UPCGNode* InOutNode, TArray<TObjectPtr<UPCGPin>>& InputPins, TArray<TObjectPtr<UPCGPin>>& OutputPins);
 	/** Any final migration/recovery that can be done after pins are finalized. This function should also set DataVersion to LatestVersion. */
-	virtual void ApplyDeprecation(UPCGNode* InOutNode);
+	UE_API virtual void ApplyDeprecation(UPCGNode* InOutNode);
 
 	/** If settings require structural changes, this will apply them */
-	virtual void ApplyStructuralDeprecation(UPCGNode* InOutNode);
+	UE_API virtual void ApplyStructuralDeprecation(UPCGNode* InOutNode);
 
 	virtual FName GetDefaultNodeName() const { return NAME_None; }
 	virtual FText GetDefaultNodeTitle() const { return FText::FromName(GetDefaultNodeName()); }
@@ -278,7 +280,7 @@ public:
 	virtual bool GetPinLabelStyle(const UPCGPin* InPin, FName& OutLabelStyle) const { return false; }
 
 	/** Can override to add a custom icon next to the pin label (and an optional tooltip). Return false if no override is available. */
-	virtual bool GetPinExtraIcon(const UPCGPin* InPin, FName& OutExtraIcon, FText& OutTooltip) const;
+	UE_API virtual bool GetPinExtraIcon(const UPCGPin* InPin, FName& OutExtraIcon, FText& OutTooltip) const;
 
 	/** Derived classes must implement this to communicate dependencies that are known statically. */
 	virtual void GetStaticTrackedKeys(FPCGSelectionKeyToSettingsMap& OutKeysToSettings, TArray<TObjectPtr<const UPCGGraph>>& OutVisitedGraphs) const {}
@@ -290,7 +292,7 @@ public:
 	 *  ie. returning a blueprint instance will open the given blueprint in its editor.
 	 *  By default, it will return the underlying class, to try to jump to its header in code
      */
-	virtual UObject* GetJumpTargetForDoubleClick() const;
+	UE_API virtual UObject* GetJumpTargetForDoubleClick() const;
 
 	/** Return preconfigured info that will be filled in the editor palette action, allowing to create pre-configured settings */
 	virtual TArray<FPCGPreConfiguredSettingsInfo> GetPreconfiguredInfo() const { return {}; }
@@ -305,7 +307,7 @@ public:
 	virtual TArray<FPCGPreconfiguredInfo> GetConversionInfo() const { return {}; }
 
 	/** Perform post-operations when an editor node is copied */
-	virtual void PostPaste();
+	UE_API virtual void PostPaste();
 #endif // WITH_EDITOR
 
 	/** Derived classes can implement this to expose additional information or context, such as an asset in use by the node. */
@@ -315,13 +317,13 @@ public:
 	virtual bool HasFlippedTitleLines() const { return false; }
 
 	/** Returns true if the given property is overridden by graph parameters */
-	virtual bool IsPropertyOverriddenByPin(const FProperty* InProperty) const;
+	UE_API virtual bool IsPropertyOverriddenByPin(const FProperty* InProperty) const;
 
 	/** Returns true if the base property, given by name, is overridden by graph parameters */
-	virtual bool IsPropertyOverriddenByPin(FName PropertyName) const;
+	UE_API virtual bool IsPropertyOverriddenByPin(FName PropertyName) const;
 
 	/** Returns true if the property, given by chain of property names, is overridden by graph parameters */
-	virtual bool IsPropertyOverriddenByPin(const TArrayView<const FName>& PropertyNameChain) const;
+	UE_API virtual bool IsPropertyOverriddenByPin(const TArrayView<const FName>& PropertyNameChain) const;
 
 	virtual void ApplyPreconfiguredSettings(const FPCGPreConfiguredSettingsInfo& PreconfigureInfo) {}
 
@@ -348,12 +350,12 @@ public:
 	virtual bool OnlyPassThroughOneEdgeWhenDisabled() const { return false; }
 
 	/** By default, the first non-Advanced pin that supports pass-through will be selected. */
-	virtual bool DoesPinSupportPassThrough(UPCGPin* InPin) const;
+	UE_API virtual bool DoesPinSupportPassThrough(UPCGPin* InPin) const;
 
 	/** Returns the current pin types, which can either be the static types from the pin properties, or a dynamic type based on connected edges.
 	* By default we set output pin types to the union of the default input pin incident edge types, if it is dynamic and the default input exists.
 	*/
-	virtual EPCGDataType GetCurrentPinTypes(const UPCGPin* InPin) const;
+	UE_API virtual EPCGDataType GetCurrentPinTypes(const UPCGPin* InPin) const;
 
 	/** Convert this owning node into the corresponding target. Returns true upon success. */
 	virtual bool ConvertNode(const FPCGPreconfiguredInfo& ConversionInfo) { return false; }
@@ -388,8 +390,8 @@ public:
 	const UPCGSettings* OriginalSettings = nullptr;
 
 	// Returns an array of all the input pin properties. You should not add manually a "params" pin, it is handled automatically by FillOverridableParamsPins
-	virtual TArray<FPCGPinProperties> InputPinProperties() const;
-	virtual TArray<FPCGPinProperties> OutputPinProperties() const;
+	UE_API virtual TArray<FPCGPinProperties> InputPinProperties() const;
+	UE_API virtual TArray<FPCGPinProperties> OutputPinProperties() const;
 
 protected:
 	// BP version since EPCGDataType is uint32 and BP only supports uint8 enums.
@@ -397,37 +399,37 @@ protected:
 	UFUNCTION(BlueprintCallable, Category = "Settings|DynamicPins", DisplayName = "Get Type Union Of Incident Edges")
 	int32 BP_GetTypeUnionOfIncidentEdges(const FName& PinLabel) const { return static_cast<int32>(GetTypeUnionOfIncidentEdges(PinLabel)); }
 
-	virtual FPCGElementPtr CreateElement() const PURE_VIRTUAL(UPCGSettings::CreateElement, return nullptr;);
+	UE_API virtual FPCGElementPtr CreateElement() const PURE_VIRTUAL(UPCGSettings::CreateElement, return nullptr;);
 
 	/** An additional custom version number that external system users can use to track versions. This version will be serialized into the asset and will be provided by UserDataVersion after load. */
 	virtual FGuid GetUserCustomVersionGuid() { return FGuid(); }
 
 	/** Can be overriden by child class if they ever got renamed to avoid changing the default seed for this one. Otherwise default is hash of the class name. */
-	virtual uint32 GetTypeNameHash() const;
+	UE_API virtual uint32 GetTypeNameHash() const;
 
 	/** Can be overriden by child class if some fixup code needs to run after duplication in the context of FPCGContext::InitializeSettings */
 	virtual void OnOverrideSettingsDuplicatedInternal(bool bSkippedPostLoad) {};
 
 #if WITH_EDITOR
-	virtual void PostEditChangeProperty(FPropertyChangedEvent& PropertyChangedEvent) override;
+	UE_API virtual void PostEditChangeProperty(FPropertyChangedEvent& PropertyChangedEvent) override;
 
 	UE_DEPRECATED(5.4, "IsStructuralProperty is deprecated, return EPCGChangeType::Structural from GetChangeTypeForProperty instead.")
 	virtual bool IsStructuralProperty(const FName& InPropertyName) const { return false; }
 
 	/** Gets the change impact for a given property. Can be used to signal structural or cosmetic node changes for example - calls the named version by default. */
-	virtual EPCGChangeType GetChangeTypeForProperty(FPropertyChangedEvent& PropertyChangedEvent) const;
+	UE_API virtual EPCGChangeType GetChangeTypeForProperty(FPropertyChangedEvent& PropertyChangedEvent) const;
 
 	/** Gets the change impact for a given property. Can be used to signal structural or cosmetic node changes for example. */
-	virtual EPCGChangeType GetChangeTypeForProperty(const FName& InPropertyName) const;
+	UE_API virtual EPCGChangeType GetChangeTypeForProperty(const FName& InPropertyName) const;
 
 	/** Method that can be called to dirty the cache data from this settings objects if the operator== does not allow to detect changes */
-	void DirtyCache();
+	UE_API void DirtyCache();
 
 	// We need the more complex function (with PropertyChain) to detect child properties in structs, if they are overridable
-	virtual bool CanEditChange(const FEditPropertyChain& InPropertyChain) const override;
+	UE_API virtual bool CanEditChange(const FEditPropertyChain& InPropertyChain) const override;
 
 	// Passthrough for the simpler method, to avoid modifying the child settings already overriding this method.
-	virtual bool CanEditChange(const FProperty* InProperty) const override;
+	UE_API virtual bool CanEditChange(const FProperty* InProperty) const override;
 #endif
 
 	UE_DEPRECATED(5.5, "Implement the PCGSettings virtual UseSeed() override.")
@@ -435,8 +437,8 @@ protected:
 	bool bUseSeed = false;
 
 	/** Methods to remove boilerplate code across settings */
-	TArray<FPCGPinProperties> DefaultPointInputPinProperties() const;
-	TArray<FPCGPinProperties> DefaultPointOutputPinProperties() const;
+	UE_API TArray<FPCGPinProperties> DefaultPointInputPinProperties() const;
+	UE_API TArray<FPCGPinProperties> DefaultPointOutputPinProperties() const;
 
 public:
 #if WITH_EDITORONLY_DATA
@@ -460,18 +462,18 @@ public:
 	virtual bool HasOverridableParams() const { return !CachedOverridableParams.IsEmpty(); }
 
 	/** Checks if a label matches an overridable param */
-	virtual bool HasOverridableParam(FName InParamName) const;
+	UE_API virtual bool HasOverridableParam(FName InParamName) const;
 
 	/** Check if we need to hook the output of the pre-task to this. One use is to compute overrides in the subgraph element and pass the overrides as data, to all nodes that needs it. */
 	virtual bool RequiresDataFromPreTask() const { return false; }
 
 protected:
 	/** Iterate over OverridableParams to automatically add param pins to the list. */
-	void FillOverridableParamsPins(TArray<FPCGPinProperties>& OutPins) const;
+	UE_API void FillOverridableParamsPins(TArray<FPCGPinProperties>& OutPins) const;
 
 #if WITH_EDITOR
 	/** Can be overridden to add more custom params (like in BP). */
-	virtual TArray<FPCGSettingsOverridableParam> GatherOverridableParams() const;
+	UE_API virtual TArray<FPCGSettingsOverridableParam> GatherOverridableParams() const;
 #endif // WITH_EDITOR
 
 	/** Called after intialization to construct the list of all overridable params. 
@@ -479,13 +481,13 @@ protected:
 	* Can request a "reset" if something changed in the settings.
 	* And then, in Editor and Runtime, will gather the FProperty*.
 	*/
-	void InitializeCachedOverridableParams(bool bReset = false);
+	UE_API void InitializeCachedOverridableParams(bool bReset = false);
 
 	/**
 	* There is a weird issue where the BP class is not set correctly in some Server cases.
 	* We can try to recover if the PropertyClass is null.
 	*/
-	virtual void FixingOverridableParamPropertyClass(FPCGSettingsOverridableParam& Param) const;
+	UE_API virtual void FixingOverridableParamPropertyClass(FPCGSettingsOverridableParam& Param) const;
 
 	/** Needs to be serialized since property metadata (used to populate this array) is not available at runtime. */
 	UPROPERTY()
@@ -526,7 +528,7 @@ protected:
 
 private:
 	/** Adds the universal Execution Dependency pin to the node's pin properties. */
-	void AddExecutionDependencyPin(TArray<FPCGPinProperties>& OutPins) const;
+	UE_API void AddExecutionDependencyPin(TArray<FPCGPinProperties>& OutPins) const;
 
 protected:
 	/** [EXPERIMENTAL] Note: This variable is subject to change without deprecation. */
@@ -552,7 +554,7 @@ public:
 
 private:
 	/** Calculate Crc for these settings and save it. */
-	void CacheCrc();
+	UE_API void CacheCrc();
 
 #if WITH_EDITORONLY_DATA
 	/** The cached Crc for these settings. */
@@ -588,8 +590,8 @@ public:
 #endif
 };
 
-UCLASS(BlueprintType, ClassGroup = (Procedural))
-class PCG_API UPCGSettingsInstance : public UPCGSettingsInterface
+UCLASS(MinimalAPI, BlueprintType, ClassGroup = (Procedural))
+class UPCGSettingsInstance : public UPCGSettingsInterface
 {
 	GENERATED_BODY()
 
@@ -600,22 +602,22 @@ public:
 	// ~End UPCGSettingsInterface
 
 	// ~Begin UObject interface
-	virtual void PostLoad() override;
-	virtual void BeginDestroy() override;
-	virtual void PostEditImport() override;
+	UE_API virtual void PostLoad() override;
+	UE_API virtual void BeginDestroy() override;
+	UE_API virtual void PostEditImport() override;
 
 #if WITH_EDITOR
-	virtual void PostEditChangeProperty(FPropertyChangedEvent& PropertyChangedEvent) override;
+	UE_API virtual void PostEditChangeProperty(FPropertyChangedEvent& PropertyChangedEvent) override;
 #endif
 	// ~End UObject interface
 
-	void SetSettings(UPCGSettings* InSettings);
+	UE_API void SetSettings(UPCGSettings* InSettings);
 
 protected:
 #if WITH_EDITOR
-	void OnSettingsChanged(UPCGSettings* InSettings, EPCGChangeType ChangeType);
+	UE_API void OnSettingsChanged(UPCGSettings* InSettings, EPCGChangeType ChangeType);
 
-	void PostSettingsChanged();
+	UE_API void PostSettingsChanged();
 #endif
 
 #if WITH_EDITORONLY_DATA
@@ -650,3 +652,5 @@ protected:
 	virtual bool SupportsGPUResidentData(FPCGContext* InContext) const override { return true; }
 	virtual bool SupportsBasePointDataInputs(FPCGContext* InContext) const override { return true; }
 };
+
+#undef UE_API
diff --git a/Engine/Plugins/Runtime/Harmonix/Source/HarmonixMidi/Public/HarmonixMidi/MidiConstants.h b/Engine/Plugins/Runtime/Harmonix/Source/HarmonixMidi/Public/HarmonixMidi/MidiConstants.h
--- a/Engine/Plugins/Runtime/Harmonix/Source/HarmonixMidi/Public/HarmonixMidi/MidiConstants.h
+++ b/Engine/Plugins/Runtime/Harmonix/Source/HarmonixMidi/Public/HarmonixMidi/MidiConstants.h
@@ -93,23 +93,23 @@ namespace Harmonix::Midi::Constants
 	inline constexpr uint8 GPitch    = 0xe0; // Pitch Wheel Change
 	inline constexpr uint8 GSystem   = 0xf0; // System
 
-	HARMONIXMIDI_API inline bool  IsStatus(uint8 Byte)     { return (Byte & GStatusBitMask) > 0;  }
-	HARMONIXMIDI_API inline uint8 GetType(uint8 Status)    { check(IsStatus(Status)); return Status & GMessageTypeMask; }
+	inline bool  IsStatus(uint8 Byte)     { return (Byte & GStatusBitMask) > 0;  }
+	inline uint8 GetType(uint8 Status)    { check(IsStatus(Status)); return Status & GMessageTypeMask; }
 
-	HARMONIXMIDI_API inline bool  IsNoteOff(uint8 Status)  { return GetType(Status) == GNoteOff;  }
-	HARMONIXMIDI_API inline bool  IsNoteOn(uint8 Status)   { return GetType(Status) == GNoteOn;   }
-	HARMONIXMIDI_API inline bool  IsPolyPres(uint8 Status) { return GetType(Status) == GPolyPres; }
-	HARMONIXMIDI_API inline bool  IsControl(uint8 Status)  { return GetType(Status) == GControl;  }
-	HARMONIXMIDI_API inline bool  IsProgram(uint8 Status)  { return GetType(Status) == GProgram;  }
-	HARMONIXMIDI_API inline bool  IsChanPres(uint8 Status) { return GetType(Status) == GChanPres; }
-	HARMONIXMIDI_API inline bool  IsPitch(uint8 Status)    { return GetType(Status) == GPitch;    }
-	HARMONIXMIDI_API inline bool  IsSystem(uint8 Status)   { return GetType(Status) == GSystem;   }
+	inline bool  IsNoteOff(uint8 Status)  { return GetType(Status) == GNoteOff;  }
+	inline bool  IsNoteOn(uint8 Status)   { return GetType(Status) == GNoteOn;   }
+	inline bool  IsPolyPres(uint8 Status) { return GetType(Status) == GPolyPres; }
+	inline bool  IsControl(uint8 Status)  { return GetType(Status) == GControl;  }
+	inline bool  IsProgram(uint8 Status)  { return GetType(Status) == GProgram;  }
+	inline bool  IsChanPres(uint8 Status) { return GetType(Status) == GChanPres; }
+	inline bool  IsPitch(uint8 Status)    { return GetType(Status) == GPitch;    }
+	inline bool  IsSystem(uint8 Status)   { return GetType(Status) == GSystem;   }
 
-	HARMONIXMIDI_API inline uint8 GetChannel(uint8 Status) { check(IsStatus(Status));	check(!IsSystem(Status)); return Status & GChannelMask; }
+	inline uint8 GetChannel(uint8 Status) { check(IsStatus(Status));	check(!IsSystem(Status)); return Status & GChannelMask; }
 
 	// turns MIDI style tempo (microseconds per quarter note) into BPM...
-	HARMONIXMIDI_API inline float MidiTempoToBPM(int32 UsPerQuarterNote) { return UsPerQuarterNote== 0 ? 0 : 60000000.0f/(float)UsPerQuarterNote; }
-	HARMONIXMIDI_API inline int32 BPMToMidiTempo(float Bpm)       { return Bpm == 0.0 ? 0 : (int32)(60000000.0f / (float)Bpm); }
+	inline float MidiTempoToBPM(int32 UsPerQuarterNote) { return UsPerQuarterNote== 0 ? 0 : 60000000.0f/(float)UsPerQuarterNote; }
+	inline int32 BPMToMidiTempo(float Bpm)       { return Bpm == 0.0 ? 0 : (int32)(60000000.0f / (float)Bpm); }
 	HARMONIXMIDI_API int32 SubdivisionToMidiTicks(const EMidiClockSubdivisionQuantization Division, const FTimeSignature& TimeSignature, const int32 TicksPerQuarterNote);
 	HARMONIXMIDI_API float SubdivisionToBeats(EMidiClockSubdivisionQuantization Subdivision, const FTimeSignature& TimeSignature);
 
diff --git a/Engine/Plugins/Runtime/ModelViewViewModel/Source/ModelViewViewModel/Public/ViewModel/MVVMInstancedViewModelGeneratedClass.h b/Engine/Plugins/Runtime/ModelViewViewModel/Source/ModelViewViewModel/Public/ViewModel/MVVMInstancedViewModelGeneratedClass.h
--- a/Engine/Plugins/Runtime/ModelViewViewModel/Source/ModelViewViewModel/Public/ViewModel/MVVMInstancedViewModelGeneratedClass.h
+++ b/Engine/Plugins/Runtime/ModelViewViewModel/Source/ModelViewViewModel/Public/ViewModel/MVVMInstancedViewModelGeneratedClass.h
@@ -7,29 +7,31 @@
 #include "UObject/Package.h"
 #include "MVVMInstancedViewModelGeneratedClass.generated.h"
 
+#define UE_API MODELVIEWVIEWMODEL_API
+
 /**
  *
  */
-UCLASS()
-class MODELVIEWVIEWMODEL_API UMVVMInstancedViewModelGeneratedClass : public UBlueprintGeneratedClass
+UCLASS(MinimalAPI)
+class UMVVMInstancedViewModelGeneratedClass : public UBlueprintGeneratedClass
 {
 	GENERATED_BODY()
 
 public:
-	virtual void Link(FArchive& Ar, bool bRelinkExistingProperties) override;
+	UE_API virtual void Link(FArchive& Ar, bool bRelinkExistingProperties) override;
 	
 #if WITH_EDITOR
-	virtual UClass* GetAuthoritativeClass() override;
-	virtual void PurgeClass(bool bRecompilingOnLoad) override;
-	void AddNativeRepNotifyFunction(UFunction* Function, const FProperty* Property);
-	void PurgeNativeRepNotifyFunctions();
+	UE_API virtual UClass* GetAuthoritativeClass() override;
+	UE_API virtual void PurgeClass(bool bRecompilingOnLoad) override;
+	UE_API void AddNativeRepNotifyFunction(UFunction* Function, const FProperty* Property);
+	UE_API void PurgeNativeRepNotifyFunctions();
 #endif
 
 public:
-	DECLARE_FUNCTION(K2_CallNativeOnRep);
+	UE_API DECLARE_FUNCTION(K2_CallNativeOnRep);
 
-	void BroadcastFieldValueChanged(UObject* Object, const FProperty* Property);
-	virtual void OnPropertyReplicated(UObject* Object, const FProperty* Property);
+	UE_API void BroadcastFieldValueChanged(UObject* Object, const FProperty* Property);
+	UE_API virtual void OnPropertyReplicated(UObject* Object, const FProperty* Property);
 
 private:
 	UPROPERTY()
@@ -37,3 +39,5 @@ private:
 
 	TMap<TObjectKey<UFunction>, const FProperty*> OnRepToPropertyMap;
 };
+
+#undef UE_API
diff --git a/Engine/Source/Editor/KismetCompiler/Public/KismetCompiledFunctionContext.h b/Engine/Source/Editor/KismetCompiler/Public/KismetCompiledFunctionContext.h
--- a/Engine/Source/Editor/KismetCompiler/Public/KismetCompiledFunctionContext.h
+++ b/Engine/Source/Editor/KismetCompiler/Public/KismetCompiledFunctionContext.h
@@ -102,9 +102,9 @@ public:
 	//Does this function use requires FlowStack ?
 	bool bUseFlowStack;
 public:
-	FKismetFunctionContext(FCompilerResultsLog& InMessageLog, const UEdGraphSchema_K2* InSchema, UBlueprintGeneratedClass* InNewClass, UBlueprint* InBlueprint);
+	KISMETCOMPILER_API FKismetFunctionContext(FCompilerResultsLog& InMessageLog, const UEdGraphSchema_K2* InSchema, UBlueprintGeneratedClass* InNewClass, UBlueprint* InBlueprint);
 
-	~FKismetFunctionContext();
+	KISMETCOMPILER_API ~FKismetFunctionContext();
 
 	void SetExternalNetNameMap(FNetNameMapping* NewMap);
 
diff --git a/Engine/Source/Runtime/TypedElementFramework/Public/Elements/Framework/TypedElementData.h b/Engine/Source/Runtime/TypedElementFramework/Public/Elements/Framework/TypedElementData.h
--- a/Engine/Source/Runtime/TypedElementFramework/Public/Elements/Framework/TypedElementData.h
+++ b/Engine/Source/Runtime/TypedElementFramework/Public/Elements/Framework/TypedElementData.h
@@ -14,8 +14,8 @@
  * Macro to declare the required RTTI data for types representing element data.
  * @note Place this in the public section of your type declaration.
  */
-#define UE_DECLARE_TYPED_ELEMENT_DATA_RTTI(ELEMENT_DATA_TYPE)						\
-	static FTypedHandleTypeId Private_RegisteredTypeId;								\
+#define UE_DECLARE_TYPED_ELEMENT_DATA_RTTI(ELEMENT_DATA_TYPE, ...)					\
+	static __VA_ARGS__ FTypedHandleTypeId Private_RegisteredTypeId;				\
 	static FTypedHandleTypeId StaticTypeId() { return Private_RegisteredTypeId; }	\
 	static FName StaticTypeName() { static const FName TypeName = #ELEMENT_DATA_TYPE; return TypeName; }
 
