NNERuntimeBasicCpu: Added some more robustness to empty matrices.

diff --git a/Engine/Plugins/Experimental/NNERuntimeBasicCpu/Source/NNERuntimeBasicCpu/Private/NNERuntimeBasicCpuModel.cpp b/Engine/Plugins/Experimental/NNERuntimeBasicCpu/Source/NNERuntimeBasicCpu/Private/NNERuntimeBasicCpuModel.cpp
--- a/Engine/Plugins/Experimental/NNERuntimeBasicCpu/Source/NNERuntimeBasicCpu/Private/NNERuntimeBasicCpuModel.cpp
+++ b/Engine/Plugins/Experimental/NNERuntimeBasicCpu/Source/NNERuntimeBasicCpu/Private/NNERuntimeBasicCpuModel.cpp
@@ -29,9 +29,11 @@
 #endif
 
 #if NNE_RUNTIME_BASIC_ENABLE_ALIASING_CHECK
-#define NNE_RUNTIME_BASIC_CHECK_ALIASING(X, Y) check((!X && !Y) || (X != Y));
+#define NNE_RUNTIME_BASIC_CHECK_ALIASING_SIZE(X, Y, Size) check(X == nullptr || Y == nullptr || Size == 0 || (X != Y));
+#define NNE_RUNTIME_BASIC_CHECK_ALIASING_SIZES(X, Y, XSize, YSize) check(X == nullptr || Y == nullptr || XSize == 0 || YSize == 0 || (X != Y));
 #else
-#define NNE_RUNTIME_BASIC_CHECK_ALIASING(...)
+#define NNE_RUNTIME_BASIC_CHECK_ALIASING_SIZE(...)
+#define NNE_RUNTIME_BASIC_CHECK_ALIASING_SIZES(...)
 #endif
 
 namespace UE::NNE::RuntimeBasic
@@ -227,7 +229,7 @@ namespace UE::NNE::RuntimeBasic
 		{
 			NNE_RUNTIME_BASIC_TRACE_SCOPE(NNE::RuntimeBasic::Private::OperatorCopy);
 
-			NNE_RUNTIME_BASIC_CHECK_ALIASING(Output, Input);
+			NNE_RUNTIME_BASIC_CHECK_ALIASING_SIZE(Output, Input, InputOutputSize);
 
 #if NNE_RUNTIME_BASIC_ENABLE_ISPC
 			ispc::NNERuntimeBasicCPUOperatorCopy(
@@ -258,7 +260,7 @@ namespace UE::NNE::RuntimeBasic
 		{
 			NNE_RUNTIME_BASIC_TRACE_SCOPE(NNE::RuntimeBasic::Private::OperatorAddInplace);
 
-			NNE_RUNTIME_BASIC_CHECK_ALIASING(Output, Input);
+			NNE_RUNTIME_BASIC_CHECK_ALIASING_SIZE(Output, Input, InputOutputSize);
 
 #if NNE_RUNTIME_BASIC_ENABLE_ISPC
 			ispc::NNERuntimeBasicCPUOperatorAddInplace(
@@ -290,7 +292,7 @@ namespace UE::NNE::RuntimeBasic
 		{
 			NNE_RUNTIME_BASIC_TRACE_SCOPE(NNE::RuntimeBasic::Private::OperatorTile);
 
-			NNE_RUNTIME_BASIC_CHECK_ALIASING(Output, Input);
+			NNE_RUNTIME_BASIC_CHECK_ALIASING_SIZES(Output, Input, Repeats * InputSize, InputSize);
 
 #if NNE_RUNTIME_BASIC_ENABLE_ISPC
 			ispc::NNERuntimeBasicCPUOperatorTile(
@@ -327,7 +329,7 @@ namespace UE::NNE::RuntimeBasic
 		{
 			NNE_RUNTIME_BASIC_TRACE_SCOPE(NNE::RuntimeBasic::Private::OperatorNormalize);
 
-			NNE_RUNTIME_BASIC_CHECK_ALIASING(Output, Input);
+			NNE_RUNTIME_BASIC_CHECK_ALIASING_SIZE(Output, Input, InputOutputSize);
 
 #if NNE_RUNTIME_BASIC_ENABLE_ISPC
 			ispc::NNERuntimeBasicCPUOperatorNormalize(
@@ -362,7 +364,7 @@ namespace UE::NNE::RuntimeBasic
 		{
 			NNE_RUNTIME_BASIC_TRACE_SCOPE(NNE::RuntimeBasic::Private::OperatorDenormalize);
 
-			NNE_RUNTIME_BASIC_CHECK_ALIASING(Output, Input);
+			NNE_RUNTIME_BASIC_CHECK_ALIASING_SIZE(Output, Input, InputOutputSize);
 
 #if NNE_RUNTIME_BASIC_ENABLE_ISPC
 			ispc::NNERuntimeBasicCPUOperatorDenormalize(
@@ -397,7 +399,7 @@ namespace UE::NNE::RuntimeBasic
 		{
 			NNE_RUNTIME_BASIC_TRACE_SCOPE(NNE::RuntimeBasic::Private::OperatorClamp);
 
-			NNE_RUNTIME_BASIC_CHECK_ALIASING(Output, Input);
+			NNE_RUNTIME_BASIC_CHECK_ALIASING_SIZE(Output, Input, InputOutputSize);
 
 #if NNE_RUNTIME_BASIC_ENABLE_ISPC
 			ispc::NNERuntimeBasicCPUOperatorClamp(
@@ -433,7 +435,7 @@ namespace UE::NNE::RuntimeBasic
 		{
 			NNE_RUNTIME_BASIC_TRACE_SCOPE(NNE::RuntimeBasic::Private::OperatorLinear);
 
-			NNE_RUNTIME_BASIC_CHECK_ALIASING(Output, Input);
+			NNE_RUNTIME_BASIC_CHECK_ALIASING_SIZES(Output, Input, OutputSize, InputSize);
 
 #if NNE_RUNTIME_BASIC_ENABLE_ISPC
 			ispc::NNERuntimeBasicCPUOperatorLinear(
@@ -488,7 +490,7 @@ namespace UE::NNE::RuntimeBasic
 		{
 			NNE_RUNTIME_BASIC_TRACE_SCOPE(NNE::RuntimeBasic::Private::OperatorCompressedLinear);
 
-			NNE_RUNTIME_BASIC_CHECK_ALIASING(Output, Input);
+			NNE_RUNTIME_BASIC_CHECK_ALIASING_SIZES(Output, Input, OutputSize, InputSize);
 
 #if NNE_RUNTIME_BASIC_ENABLE_ISPC
 			ispc::NNERuntimeBasicCPUOperatorCompressedLinear(
@@ -547,7 +549,7 @@ namespace UE::NNE::RuntimeBasic
 		{
 			NNE_RUNTIME_BASIC_TRACE_SCOPE(NNE::RuntimeBasic::Private::OperatorMultiLinear);
 
-			NNE_RUNTIME_BASIC_CHECK_ALIASING(Output, Input);
+			NNE_RUNTIME_BASIC_CHECK_ALIASING_SIZES(Output, Input, OutputSize, InputSize);
 
 			// For this function ispc generates slightly less efficient code than the naive C++ implementation so we
 			// don't bother calling out to the ispc version even if it is available
@@ -593,7 +595,7 @@ namespace UE::NNE::RuntimeBasic
 		{
 			NNE_RUNTIME_BASIC_TRACE_SCOPE(NNE::RuntimeBasic::Private::OperatorReLU);
 
-			NNE_RUNTIME_BASIC_CHECK_ALIASING(Output, Input);
+			NNE_RUNTIME_BASIC_CHECK_ALIASING_SIZE(Output, Input, InputOutputSize);
 
 #if NNE_RUNTIME_BASIC_ENABLE_ISPC
 			ispc::NNERuntimeBasicCPUOperatorReLU(
@@ -624,7 +626,7 @@ namespace UE::NNE::RuntimeBasic
 		{
 			NNE_RUNTIME_BASIC_TRACE_SCOPE(NNE::RuntimeBasic::Private::OperatorELU);
 
-			NNE_RUNTIME_BASIC_CHECK_ALIASING(Output, Input);
+			NNE_RUNTIME_BASIC_CHECK_ALIASING_SIZE(Output, Input, InputOutputSize);
 
 #if NNE_RUNTIME_BASIC_ENABLE_ISPC
 			ispc::NNERuntimeBasicCPUOperatorELU(
@@ -661,7 +663,7 @@ namespace UE::NNE::RuntimeBasic
 		{
 			NNE_RUNTIME_BASIC_TRACE_SCOPE(NNE::RuntimeBasic::Private::OperatorGELU);
 
-			NNE_RUNTIME_BASIC_CHECK_ALIASING(Output, Input);
+			NNE_RUNTIME_BASIC_CHECK_ALIASING_SIZE(Output, Input, InputOutputSize);
 
 #if NNE_RUNTIME_BASIC_ENABLE_ISPC
 			ispc::NNERuntimeBasicCPUOperatorGELU(
@@ -692,7 +694,7 @@ namespace UE::NNE::RuntimeBasic
 		{
 			NNE_RUNTIME_BASIC_TRACE_SCOPE(NNE::RuntimeBasic::Private::OperatorTanH);
 
-			NNE_RUNTIME_BASIC_CHECK_ALIASING(Output, Input);
+			NNE_RUNTIME_BASIC_CHECK_ALIASING_SIZE(Output, Input, InputOutputSize);
 
 #if NNE_RUNTIME_BASIC_ENABLE_ISPC
 			ispc::NNERuntimeBasicCPUOperatorTanH(
@@ -724,7 +726,7 @@ namespace UE::NNE::RuntimeBasic
 		{
 			NNE_RUNTIME_BASIC_TRACE_SCOPE(NNE::RuntimeBasic::Private::OperatorPReLU);
 
-			NNE_RUNTIME_BASIC_CHECK_ALIASING(Output, Input);
+			NNE_RUNTIME_BASIC_CHECK_ALIASING_SIZE(Output, Input, InputOutputSize);
 
 #if NNE_RUNTIME_BASIC_ENABLE_ISPC
 			ispc::NNERuntimeBasicCPUOperatorPReLU(
@@ -761,9 +763,9 @@ namespace UE::NNE::RuntimeBasic
 		{
 			NNE_RUNTIME_BASIC_TRACE_SCOPE(NNE::RuntimeBasic::Private::OperatorMemoryCellUpdateMemory);
 
-			NNE_RUNTIME_BASIC_CHECK_ALIASING(Output, RememberGate);
-			NNE_RUNTIME_BASIC_CHECK_ALIASING(Output, Memory);
-			NNE_RUNTIME_BASIC_CHECK_ALIASING(Output, Update);
+			NNE_RUNTIME_BASIC_CHECK_ALIASING_SIZE(Output, RememberGate, MemorySize);
+			NNE_RUNTIME_BASIC_CHECK_ALIASING_SIZE(Output, Memory, MemorySize);
+			NNE_RUNTIME_BASIC_CHECK_ALIASING_SIZE(Output, Update, MemorySize);
 
 #if NNE_RUNTIME_BASIC_ENABLE_ISPC
 			ispc::NNERuntimeBasicCPUOperatorMemoryCellUpdateMemory(
@@ -806,9 +808,9 @@ namespace UE::NNE::RuntimeBasic
 		{
 			NNE_RUNTIME_BASIC_TRACE_SCOPE(NNE::RuntimeBasic::Private::OperatorMemoryCellUpdateOutput);
 
-			NNE_RUNTIME_BASIC_CHECK_ALIASING(Output, PassthroughGate);
-			NNE_RUNTIME_BASIC_CHECK_ALIASING(Output, MemoryUpdate);
-			NNE_RUNTIME_BASIC_CHECK_ALIASING(Output, InputUpdate);
+			NNE_RUNTIME_BASIC_CHECK_ALIASING_SIZE(Output, PassthroughGate, OutputSize);
+			NNE_RUNTIME_BASIC_CHECK_ALIASING_SIZE(Output, MemoryUpdate, OutputSize);
+			NNE_RUNTIME_BASIC_CHECK_ALIASING_SIZE(Output, InputUpdate, OutputSize);
 
 #if NNE_RUNTIME_BASIC_ENABLE_ISPC
 			ispc::NNERuntimeBasicCPUOperatorMemoryCellUpdateOutput(
@@ -848,7 +850,7 @@ namespace UE::NNE::RuntimeBasic
 		{
 			NNE_RUNTIME_BASIC_TRACE_SCOPE(NNE::RuntimeBasic::Private::OperatorAggregateGatherElements);
 
-			NNE_RUNTIME_BASIC_CHECK_ALIASING(OutputBuffer, InputBuffer);
+			NNE_RUNTIME_BASIC_CHECK_ALIASING_SIZE(OutputBuffer, InputBuffer, ElementSize);
 
 #if NNE_RUNTIME_BASIC_ENABLE_ISPC
 			ispc::NNERuntimeBasicCPUOperatorAggregateGatherElements(
@@ -961,11 +963,11 @@ namespace UE::NNE::RuntimeBasic
 		{
 			NNE_RUNTIME_BASIC_TRACE_SCOPE(NNE::RuntimeBasic::Private::OperatorAggregateGatherFromSubLayers);
 
-			NNE_RUNTIME_BASIC_CHECK_ALIASING(QueryBuffer, KeyBuffer);
-			NNE_RUNTIME_BASIC_CHECK_ALIASING(QueryBuffer, ValueBuffer);
-			NNE_RUNTIME_BASIC_CHECK_ALIASING(KeyBuffer, ValueBuffer);
-			NNE_RUNTIME_BASIC_CHECK_ALIASING(ElementAccum, ElementNums);
-			NNE_RUNTIME_BASIC_CHECK_ALIASING(ElementAccum, ElementOffsets);
+			NNE_RUNTIME_BASIC_CHECK_ALIASING_SIZES(QueryBuffer, KeyBuffer, QuerySize, KeySize);
+			NNE_RUNTIME_BASIC_CHECK_ALIASING_SIZES(QueryBuffer, ValueBuffer, QuerySize, ValueSize);
+			NNE_RUNTIME_BASIC_CHECK_ALIASING_SIZES(KeyBuffer, ValueBuffer, KeySize, ValueSize);
+			NNE_RUNTIME_BASIC_CHECK_ALIASING_SIZE(ElementAccum, ElementNums, BatchSize);
+			NNE_RUNTIME_BASIC_CHECK_ALIASING_SIZE(ElementAccum, ElementOffsets, BatchSize);
 
 			const uint32 SubLayerNum = SubLayerBatchIndices.Num();
 
@@ -1039,8 +1041,8 @@ namespace UE::NNE::RuntimeBasic
 		{
 			NNE_RUNTIME_BASIC_TRACE_SCOPE(NNE::RuntimeBasic::Private::OperatorAggregateDotProductAttention);
 
-			check(Attention != Queries);
-			check(Attention != Keys);
+			NNE_RUNTIME_BASIC_CHECK_ALIASING_SIZE(Attention, Queries, AttentionEncodingSize);
+			NNE_RUNTIME_BASIC_CHECK_ALIASING_SIZE(Attention, Keys, AttentionEncodingSize);
 
 #if NNE_RUNTIME_BASIC_ENABLE_ISPC
 			ispc::NNERuntimeBasicCPUOperatorAggregateDotProductAttention(
@@ -1094,9 +1096,9 @@ namespace UE::NNE::RuntimeBasic
 		{
 			NNE_RUNTIME_BASIC_TRACE_SCOPE(NNE::RuntimeBasic::Private::OperatorAggregateSoftmaxPlusOneInplace);
 
-			NNE_RUNTIME_BASIC_CHECK_ALIASING(AttentionMaxs, AttentionDenoms);
-			NNE_RUNTIME_BASIC_CHECK_ALIASING(AttentionMaxs, Attention);
-			NNE_RUNTIME_BASIC_CHECK_ALIASING(Attention, AttentionDenoms);
+			NNE_RUNTIME_BASIC_CHECK_ALIASING_SIZE(AttentionMaxs, AttentionDenoms, AttentionHeadNum);
+			NNE_RUNTIME_BASIC_CHECK_ALIASING_SIZE(AttentionMaxs, Attention, AttentionHeadNum);
+			NNE_RUNTIME_BASIC_CHECK_ALIASING_SIZE(Attention, AttentionDenoms, AttentionHeadNum);
 
 			// Numerically stable soft-max computation using subtraction of the (positive) max value
 			// 
@@ -1165,8 +1167,8 @@ namespace UE::NNE::RuntimeBasic
 		{
 			NNE_RUNTIME_BASIC_TRACE_SCOPE(NNE::RuntimeBasic::Private::OperatorAggregateAttentionSum);
 
-			NNE_RUNTIME_BASIC_CHECK_ALIASING(Output, Attention);
-			NNE_RUNTIME_BASIC_CHECK_ALIASING(Output, Values);
+			NNE_RUNTIME_BASIC_CHECK_ALIASING_SIZES(Output, Attention, EncodingSize, AttentionHeadNum);
+			NNE_RUNTIME_BASIC_CHECK_ALIASING_SIZE(Output, Values, EncodingSize);
 
 #if NNE_RUNTIME_BASIC_ENABLE_ISPC
 			ispc::NNERuntimeBasicCPUOperatorAggregateAttentionSum(
@@ -1223,7 +1225,7 @@ namespace UE::NNE::RuntimeBasic
 		{
 			NNE_RUNTIME_BASIC_TRACE_SCOPE(NNE::RuntimeBasic::Private::OperatorGather);
 
-			NNE_RUNTIME_BASIC_CHECK_ALIASING(OutputBuffer, InputBuffer);
+			NNE_RUNTIME_BASIC_CHECK_ALIASING_SIZE(OutputBuffer, InputBuffer, InputOutputSize);
 
 #if NNE_RUNTIME_BASIC_ENABLE_ISPC
 			ispc::NNERuntimeBasicCPUOperatorGather(
@@ -1257,7 +1259,7 @@ namespace UE::NNE::RuntimeBasic
 		{
 			NNE_RUNTIME_BASIC_TRACE_SCOPE(NNE::RuntimeBasic::Private::OperatorScatter);
 
-			NNE_RUNTIME_BASIC_CHECK_ALIASING(OutputBuffer, InputBuffer);
+			NNE_RUNTIME_BASIC_CHECK_ALIASING_SIZE(OutputBuffer, InputBuffer, InputOutputSize);
 
 #if NNE_RUNTIME_BASIC_ENABLE_ISPC
 			ispc::NNERuntimeBasicCPUOperatorScatter(
@@ -1352,9 +1354,9 @@ namespace UE::NNE::RuntimeBasic
 		{
 			NNE_RUNTIME_BASIC_TRACE_SCOPE(NNE::RuntimeBasic::Private::OperatorGatherTopTwoSubLayerBatchIndices);
 
-			NNE_RUNTIME_BASIC_CHECK_ALIASING(BatchSubLayerIndex0, BatchSubLayerIndex1);
-			NNE_RUNTIME_BASIC_CHECK_ALIASING(BatchSubLayerWeight0, BatchSubLayerWeight1);
-			NNE_RUNTIME_BASIC_CHECK_ALIASING(BatchSubLayerOutputIndex0, BatchSubLayerOutputIndex1);
+			NNE_RUNTIME_BASIC_CHECK_ALIASING_SIZE(BatchSubLayerIndex0, BatchSubLayerIndex1, BatchSize);
+			NNE_RUNTIME_BASIC_CHECK_ALIASING_SIZE(BatchSubLayerWeight0, BatchSubLayerWeight1, BatchSize);
+			NNE_RUNTIME_BASIC_CHECK_ALIASING_SIZE(BatchSubLayerOutputIndex0, BatchSubLayerOutputIndex1, BatchSize);
 
 			check(SubLayerGateSize >= 2);
 
@@ -1419,8 +1421,8 @@ namespace UE::NNE::RuntimeBasic
 		{
 			NNE_RUNTIME_BASIC_TRACE_SCOPE(NNE::RuntimeBasic::Private::OperatorGatherTopTwoFromSubLayers);
 
-			NNE_RUNTIME_BASIC_CHECK_ALIASING(OutputBuffer, BatchSubLayerWeight0);
-			NNE_RUNTIME_BASIC_CHECK_ALIASING(OutputBuffer, BatchSubLayerWeight1);
+			NNE_RUNTIME_BASIC_CHECK_ALIASING_SIZES(OutputBuffer, BatchSubLayerWeight0, OutputBufferSize, BatchSize);
+			NNE_RUNTIME_BASIC_CHECK_ALIASING_SIZES(OutputBuffer, BatchSubLayerWeight1, OutputBufferSize, BatchSize);
 
 			for (uint32 BatchIdx = 0; BatchIdx < BatchSize; BatchIdx++)
 			{
@@ -1434,8 +1436,8 @@ namespace UE::NNE::RuntimeBasic
 				const float* RESTRICT SubLayerOutputBuffer0 = SubLayerOutputBuffer[SubLayerIndex0].GetData();
 				const float* RESTRICT SubLayerOutputBuffer1 = SubLayerOutputBuffer[SubLayerIndex1].GetData();
 
-				NNE_RUNTIME_BASIC_CHECK_ALIASING(OutputBuffer, SubLayerOutputBuffer0);
-				NNE_RUNTIME_BASIC_CHECK_ALIASING(OutputBuffer, SubLayerOutputBuffer1);
+				NNE_RUNTIME_BASIC_CHECK_ALIASING_SIZE(OutputBuffer, SubLayerOutputBuffer0, OutputBufferSize);
+				NNE_RUNTIME_BASIC_CHECK_ALIASING_SIZE(OutputBuffer, SubLayerOutputBuffer1, OutputBufferSize);
 
 				for (uint32 OutputIdx = 0; OutputIdx < OutputBufferSize; OutputIdx++)
 				{
@@ -1459,7 +1461,7 @@ namespace UE::NNE::RuntimeBasic
 		{
 			NNE_RUNTIME_BASIC_TRACE_SCOPE(NNE::RuntimeBasic::Private::OperatorLayerNorm);
 
-			NNE_RUNTIME_BASIC_CHECK_ALIASING(Output, Input);
+			NNE_RUNTIME_BASIC_CHECK_ALIASING_SIZE(Output, Input, InputOutputSize);
 
 #if NNE_RUNTIME_BASIC_ENABLE_ISPC
 			ispc::NNERuntimeBasicCPUOperatorLayerNorm(
@@ -1506,7 +1508,7 @@ namespace UE::NNE::RuntimeBasic
 		{
 			NNE_RUNTIME_BASIC_TRACE_SCOPE(NNE::RuntimeBasic::Private::OperatorLayerFiLM);
 
-			NNE_RUNTIME_BASIC_CHECK_ALIASING(InputOutput, Condition);
+			NNE_RUNTIME_BASIC_CHECK_ALIASING_SIZE(InputOutput, Condition, InputOutputSize);
 
 #if NNE_RUNTIME_BASIC_ENABLE_ISPC
 			ispc::NNERuntimeBasicCPUOperatorLayerFiLM(
