NNERuntimeBasicCpu: Added GRU Memory Cell and some extra checks for Instance data

diff --git a/Engine/Plugins/Experimental/NNERuntimeBasicCpu/Source/NNERuntimeBasicCpu/Private/NNERuntimeBasicCpu.ispc b/Engine/Plugins/Experimental/NNERuntimeBasicCpu/Source/NNERuntimeBasicCpu/Private/NNERuntimeBasicCpu.ispc
--- a/Engine/Plugins/Experimental/NNERuntimeBasicCpu/Source/NNERuntimeBasicCpu/Private/NNERuntimeBasicCpu.ispc
+++ b/Engine/Plugins/Experimental/NNERuntimeBasicCpu/Source/NNERuntimeBasicCpu/Private/NNERuntimeBasicCpu.ispc
@@ -338,6 +338,55 @@ export void NNERuntimeBasicCPUOperatorMemoryCellUpdateOutput(
 	}
 }
 
+export void NNERuntimeBasicCPUOperatorGRUMaskInput(
+	uniform float Masked[],
+	uniform const float InputMemory[],
+	uniform const float RememberGate[],
+	uniform const int BatchSize,
+	uniform const int OutputSize,
+	uniform const int InputSize,
+	uniform const int MaskedStride,
+	uniform const int InputMemoryStride,
+	uniform const int RememberGateStride)
+{
+	foreach(BatchIdx = 0 ... BatchSize, Idx = 0 ... InputSize)
+	{
+		Masked[BatchIdx * MaskedStride + Idx] = InputMemory[BatchIdx * InputMemoryStride + Idx];
+	}
+
+	foreach(BatchIdx = 0 ... BatchSize, Idx = 0 ... OutputSize)
+	{
+		Masked[BatchIdx * MaskedStride + InputSize + Idx] =
+			NNERuntimeBasicCPUSigmoid(RememberGate[BatchIdx * RememberGateStride + Idx]) *
+			InputMemory[BatchIdx * InputMemoryStride + InputSize + Idx];
+	}
+}
+
+export void NNERuntimeBasicCPUOperatorGRUCellUpdateOutput(
+	uniform float OutputMemory[],
+	uniform const float InputMemory[],
+	uniform const float UpdateGate[],
+	uniform const float ActivationGate[],
+	uniform const int BatchSize,
+	uniform const int OutputSize,
+	uniform const int InputSize,
+	uniform const int OutputMemoryStride,
+	uniform const int InputMemoryStride,
+	uniform const int UpdateGateStride,
+	uniform const int ActivationGateStride)
+{
+	foreach(BatchIdx = 0 ... BatchSize, Idx = 0 ... OutputSize)
+	{
+		const float MemoryValue = InputMemory[BatchIdx * InputMemoryStride + InputSize + Idx];
+		const float UpdateGateValue = NNERuntimeBasicCPUSigmoid(UpdateGate[BatchIdx * UpdateGateStride + Idx]);
+		const float ActivationGateValue = NNERuntimeBasicCPUTanH(ActivationGate[BatchIdx * ActivationGateStride + Idx]);
+		const float OutputValue = (1.0f - UpdateGateValue) * MemoryValue + UpdateGateValue * ActivationGateValue;
+
+		OutputMemory[BatchIdx * OutputMemoryStride + Idx] = OutputValue;
+		OutputMemory[BatchIdx * OutputMemoryStride + OutputSize + Idx] = OutputValue;
+	}
+}
+
 export void NNERuntimeBasicCPUOperatorAggregateGatherElements(
 	uniform float OutputBuffer[],
 	uniform const float InputBuffer[],
diff --git a/Engine/Plugins/Experimental/NNERuntimeBasicCpu/Source/NNERuntimeBasicCpu/Private/NNERuntimeBasicCpuModel.cpp b/Engine/Plugins/Experimental/NNERuntimeBasicCpu/Source/NNERuntimeBasicCpu/Private/NNERuntimeBasicCpuModel.cpp
--- a/Engine/Plugins/Experimental/NNERuntimeBasicCpu/Source/NNERuntimeBasicCpu/Private/NNERuntimeBasicCpuModel.cpp
+++ b/Engine/Plugins/Experimental/NNERuntimeBasicCpu/Source/NNERuntimeBasicCpu/Private/NNERuntimeBasicCpuModel.cpp
@@ -839,6 +839,102 @@ namespace UE::NNE::RuntimeBasic
 #endif
 		}
 
+		static inline void OperatorGRUMaskInput(
+			float* RESTRICT Masked,
+			const float* RESTRICT InputMemory,
+			const float* RESTRICT RememberGate,
+			const uint32 BatchSize,
+			const uint32 OutputSize,
+			const uint32 InputSize,
+			const uint32 MaskedStride,
+			const uint32 InputMemoryStride,
+			const uint32 RememberGateStride)
+		{
+			NNE_RUNTIME_BASIC_TRACE_SCOPE(NNE::RuntimeBasic::Private::OperatorGRUMaskInput);
+
+			NNE_RUNTIME_BASIC_CHECK_ALIASING_SIZE(Masked, InputMemory, InputSize + OutputSize);
+			NNE_RUNTIME_BASIC_CHECK_ALIASING_SIZES(Masked, RememberGate, InputSize + OutputSize, OutputSize);
+
+#if NNE_RUNTIME_BASIC_ENABLE_ISPC
+			ispc::NNERuntimeBasicCPUOperatorGRUMaskInput(
+				Masked,
+				InputMemory,
+				RememberGate,
+				BatchSize,
+				OutputSize,
+				InputSize,
+				MaskedStride,
+				InputMemoryStride,
+				RememberGateStride);
+#else
+			for (uint32 BatchIdx = 0; BatchIdx < BatchSize; BatchIdx++)
+			{
+				for (uint32 Idx = 0; Idx < InputSize; Idx++)
+				{
+					Masked[BatchIdx * MaskedStride + Idx] = InputMemory[BatchIdx * InputMemoryStride + Idx];
+				}
+
+				for (uint32 Idx = 0; Idx < OutputSize; Idx++)
+				{
+					Masked[BatchIdx * MaskedStride + InputSize + Idx] = 
+						Sigmoid(RememberGate[BatchIdx * RememberGateStride + Idx]) * 
+						InputMemory[BatchIdx * InputMemoryStride + InputSize + Idx];
+				}
+			}
+#endif
+
+		}
+
+
+		static inline void OperatorGRUCellUpdateOutput(
+			float* RESTRICT OutputMemory,
+			const float* RESTRICT InputMemory,
+			const float* RESTRICT UpdateGate ,
+			const float* RESTRICT ActivationGate,
+			const uint32 BatchSize,
+			const uint32 OutputSize,
+			const uint32 InputSize,
+			const uint32 OutputMemoryStride,
+			const uint32 InputMemoryStride,
+			const uint32 UpdateGateStride,
+			const uint32 ActivationGateStride)
+		{
+			NNE_RUNTIME_BASIC_TRACE_SCOPE(NNE::RuntimeBasic::Private::OperatorGRUCellUpdateOutput);
+
+			NNE_RUNTIME_BASIC_CHECK_ALIASING_SIZES(OutputMemory, InputMemory, OutputSize + OutputSize, InputSize + OutputSize);
+			NNE_RUNTIME_BASIC_CHECK_ALIASING_SIZES(OutputMemory, UpdateGate, OutputSize + OutputSize, OutputSize);
+			NNE_RUNTIME_BASIC_CHECK_ALIASING_SIZES(OutputMemory, ActivationGate, OutputSize + OutputSize, OutputSize);
+
+#if NNE_RUNTIME_BASIC_ENABLE_ISPC
+			ispc::NNERuntimeBasicCPUOperatorGRUCellUpdateOutput(
+				OutputMemory,
+				InputMemory,
+				UpdateGate,
+				ActivationGate,
+				BatchSize,
+				OutputSize,
+				InputSize,
+				OutputMemoryStride,
+				InputMemoryStride,
+				UpdateGateStride,
+				ActivationGateStride);
+#else
+			for (uint32 BatchIdx = 0; BatchIdx < BatchSize; BatchIdx++)
+			{
+				for (uint32 Idx = 0; Idx < OutputSize; Idx++)
+				{
+					const float MemoryValue = InputMemory[BatchIdx * InputMemoryStride + InputSize + Idx];
+					const float UpdateGateValue = Sigmoid(UpdateGate[BatchIdx * UpdateGateStride + Idx]);
+					const float ActivationGateValue = FMath::Tanh(ActivationGate[BatchIdx * ActivationGateStride + Idx]);
+					const float OutputValue = (1.0f - UpdateGateValue) * MemoryValue + UpdateGateValue * ActivationGateValue;
+
+					OutputMemory[BatchIdx * OutputMemoryStride + Idx] = OutputValue;
+					OutputMemory[BatchIdx * OutputMemoryStride + OutputSize + Idx] = OutputValue;
+				}
+			}
+#endif
+		}
+
 		static inline void OperatorAggregateGatherElements(
 			float* RESTRICT OutputBuffer,
 			const float* RESTRICT InputBuffer,
@@ -1567,6 +1663,7 @@ namespace UE::NNE::RuntimeBasic
 			Slice = 25,
 			Residual = 26,
 			FiLM = 27,
+			GRUCell = 28,
 		};
 
 		//--------------------------------------------------------------------------
@@ -1583,6 +1680,9 @@ namespace UE::NNE::RuntimeBasic
 
 			/** Indicate to this layer instance what the maximum batchsize is going to be when performing inference. */
 			virtual void SetMaxBatchSize(const uint32 MaxBatchSize) = 0;
+
+			/** Gets the type of layer this is an instance for */
+			virtual ELayerType GetLayerType() const = 0;
 		};
 
 		/**
@@ -1648,6 +1748,7 @@ namespace UE::NNE::RuntimeBasic
 			FSequenceLayerInstance(const FSequenceLayer& InSequenceLayer);
 
 			virtual void SetMaxBatchSize(const uint32 MaxBatchSize) override final;
+			virtual ELayerType GetLayerType() const override final { return ELayerType::Sequence; }
 
 			const FSequenceLayer& SequenceLayer;
 			uint32 ActivationStride = 0;
@@ -1704,13 +1805,13 @@ namespace UE::NNE::RuntimeBasic
 				NNE_RUNTIME_BASIC_TRACE_SCOPE(NNE::RuntimeBasic::Private::FSequenceLayer::Evaluate);
 				check(OutputBufferSize == GetOutputSize() && InputBufferSize == GetInputSize());
 				check(OutputBufferStride >= GetOutputSize() && InputBufferStride >= GetInputSize());
+				check(Instance && Instance->GetLayerType() == GetLayerType());
+
+				FSequenceLayerInstance* SequenceInstance = StaticCast<FSequenceLayerInstance*>(Instance);
+
 				OperatorNanCheck(InputBuffer, BatchSize, InputBufferSize, InputBufferStride);
 
 				const uint32 LayerNum = Layers.Num();
-				check(LayerNum > 0);
-
-				FSequenceLayerInstance* SequenceInstance = StaticCast<FSequenceLayerInstance*>(Instance);
-				check(SequenceInstance);
 
 				// If we just have one layer then evaluate layer directly without using intermediate storage
 
@@ -2461,6 +2562,7 @@ namespace UE::NNE::RuntimeBasic
 			FMemoryCellInstance(const FMemoryCellLayer& InMemoryCellLayer);
 
 			virtual void SetMaxBatchSize(const uint32 MaxBatchSize) override final;
+			virtual ELayerType GetLayerType() const override final { return ELayerType::MemoryCell; }
 
 			const FMemoryCellLayer& MemoryCellLayer;
 			TSharedPtr<ILayerInstance> RememberInstance;
@@ -2531,9 +2633,9 @@ namespace UE::NNE::RuntimeBasic
 				NNE_RUNTIME_BASIC_TRACE_SCOPE(NNE::RuntimeBasic::Private::FMemoryCellLayer::Evaluate);
 				check(OutputBufferSize == GetOutputSize() && InputBufferSize == GetInputSize());
 				check(OutputBufferStride >= GetOutputSize() && InputBufferStride >= GetInputSize());
+				check(Instance && Instance->GetLayerType() == GetLayerType());
 
 				FMemoryCellInstance* MemoryCellInstance = StaticCast<FMemoryCellInstance*>(Instance);
-				check(MemoryCellInstance);
 				
 				OperatorNanCheck(InputBuffer, BatchSize, InputBufferSize, InputBufferStride);
 
@@ -2642,10 +2744,10 @@ namespace UE::NNE::RuntimeBasic
 			: MemoryCellLayer(InMemoryCellLayer)
 		{
 			RememberInstance = MemoryCellLayer.RememberLayer->MakeInstance();
-			PassthroughInstance = MemoryCellLayer.RememberLayer->MakeInstance();
-			MemoryUpdateInstance = MemoryCellLayer.RememberLayer->MakeInstance();
-			OutputInputUpdateInstance = MemoryCellLayer.RememberLayer->MakeInstance();
-			OutputMemoryUpdateInstance = MemoryCellLayer.RememberLayer->MakeInstance();
+			PassthroughInstance = MemoryCellLayer.PassthroughLayer->MakeInstance();
+			MemoryUpdateInstance = MemoryCellLayer.MemoryUpdateLayer->MakeInstance();
+			OutputInputUpdateInstance = MemoryCellLayer.OutputInputUpdateLayer->MakeInstance();
+			OutputMemoryUpdateInstance = MemoryCellLayer.OutputMemoryUpdateLayer->MakeInstance();
 		}
 
 		void FMemoryCellInstance::SetMaxBatchSize(const uint32 MaxBatchSize)
@@ -2726,7 +2828,8 @@ namespace UE::NNE::RuntimeBasic
 		{
 			FConcatLayerInstance(const FConcatLayer& InConcatLayer);
 
-			void SetMaxBatchSize(const uint32 MaxBatchSize) override final;
+			virtual void SetMaxBatchSize(const uint32 MaxBatchSize) override final;
+			virtual ELayerType GetLayerType() const override final { return ELayerType::Concat; }
 
 			const FConcatLayer& ConcatLayer;
 			TArray<TSharedPtr<ILayerInstance>, TInlineAllocator<32>> Instances;
@@ -2796,9 +2899,9 @@ namespace UE::NNE::RuntimeBasic
 				NNE_RUNTIME_BASIC_TRACE_SCOPE(NNE::RuntimeBasic::Private::FConcatLayer::Evaluate);
 				check(OutputBufferSize == GetOutputSize() && InputBufferSize == GetInputSize());
 				check(OutputBufferStride >= GetOutputSize() && InputBufferStride >= GetInputSize());
+				check(Instance && Instance->GetLayerType() == GetLayerType());
 
 				FConcatLayerInstance* ConcatInstance = StaticCast<FConcatLayerInstance*>(Instance);
-				check(ConcatInstance);
 
 				OperatorNanCheck(InputBuffer, BatchSize, InputBufferSize, InputBufferStride);
 
@@ -2868,7 +2971,8 @@ namespace UE::NNE::RuntimeBasic
 		{
 			FArrayLayerInstance(const FArrayLayer& InArrayLayer);
 
-			void SetMaxBatchSize(const uint32 MaxBatchSize) override final;
+			virtual void SetMaxBatchSize(const uint32 MaxBatchSize) override final;
+			virtual ELayerType GetLayerType() const override final { return ELayerType::Array; }
 
 			const FArrayLayer& ArrayLayer;
 			TSharedPtr<ILayerInstance> Instance;
@@ -2920,9 +3024,9 @@ namespace UE::NNE::RuntimeBasic
 				NNE_RUNTIME_BASIC_TRACE_SCOPE(NNE::RuntimeBasic::Private::FArrayLayer::Evaluate);
 				check(OutputBufferSize == GetOutputSize() && InputBufferSize == GetInputSize());
 				check(OutputBufferStride >= GetOutputSize() && InputBufferStride >= GetInputSize());
+				check(Instance && Instance->GetLayerType() == GetLayerType());
 
 				FArrayLayerInstance* ArrayInstance = StaticCast<FArrayLayerInstance*>(Instance);
-				check(ArrayInstance);
 
 				OperatorNanCheck(InputBuffer, BatchSize, InputBufferSize, InputBufferStride);
 
@@ -3012,7 +3116,8 @@ namespace UE::NNE::RuntimeBasic
 		{
 			FAggregateSetLayerInstance(const FAggregateSetLayer& InAggregateSetLayer);
 
-			void SetMaxBatchSize(const uint32 MaxBatchSize) override final;
+			virtual void SetMaxBatchSize(const uint32 MaxBatchSize) override final;
+			virtual ELayerType GetLayerType() const override final { return ELayerType::AggregateSet; }
 
 			const FAggregateSetLayer& AggregateSetLayer;
 			TSharedPtr<ILayerInstance> SubLayerInstance;
@@ -3096,9 +3201,9 @@ namespace UE::NNE::RuntimeBasic
 				NNE_RUNTIME_BASIC_TRACE_SCOPE(NNE::RuntimeBasic::Private::FAggregateSetLayer::Evaluate);
 				check(OutputBufferSize == GetOutputSize() && InputBufferSize == GetInputSize());
 				check(OutputBufferStride >= GetOutputSize() && InputBufferStride >= GetInputSize());
+				check(Instance && Instance->GetLayerType() == GetLayerType());
 
 				FAggregateSetLayerInstance* AggregateSetInstance = StaticCast<FAggregateSetLayerInstance*>(Instance);
-				check(AggregateSetInstance);
 
 				OperatorNanCheck(InputBuffer, BatchSize, InputBufferSize, InputBufferStride);
 
@@ -3267,7 +3372,8 @@ namespace UE::NNE::RuntimeBasic
 		{
 			FAggregateOrExclusiveLayerInstance(const FAggregateOrExclusiveLayer& InAggregateOrExclusiveLayer);
 
-			void SetMaxBatchSize(const uint32 MaxBatchSize) override final;
+			virtual void SetMaxBatchSize(const uint32 MaxBatchSize) override final;
+			virtual ELayerType GetLayerType() const override final { return ELayerType::AggregateOrExclusive; }
 
 			const FAggregateOrExclusiveLayer& AggregateOrExclusiveLayer;
 			TArray<TSharedPtr<ILayerInstance>, TInlineAllocator<32>> SubLayerInstances;
@@ -3343,9 +3449,9 @@ namespace UE::NNE::RuntimeBasic
 				NNE_RUNTIME_BASIC_TRACE_SCOPE(NNE::RuntimeBasic::Private::FAggregateOrExclusiveLayer::Evaluate);
 				check(OutputBufferSize == GetOutputSize() && InputBufferSize == GetInputSize());
 				check(OutputBufferStride >= GetOutputSize() && InputBufferStride >= GetInputSize());
+				check(Instance && Instance->GetLayerType() == GetLayerType());
 
 				FAggregateOrExclusiveLayerInstance* AggregateOrExclusiveInstance = StaticCast<FAggregateOrExclusiveLayerInstance*>(Instance);
-				check(AggregateOrExclusiveInstance);
 
 				OperatorNanCheck(InputBuffer, BatchSize, InputBufferSize, InputBufferStride);
 
@@ -3475,7 +3581,8 @@ namespace UE::NNE::RuntimeBasic
 		{
 			FAggregateOrInclusiveLayerInstance(const FAggregateOrInclusiveLayer& InAggregateOrInclusiveLayer);
 
-			void SetMaxBatchSize(const uint32 MaxBatchSize) override final;
+			virtual void SetMaxBatchSize(const uint32 MaxBatchSize) override final;
+			virtual ELayerType GetLayerType() const override final { return ELayerType::AggregateOrInclusive; }
 
 			const FAggregateOrInclusiveLayer& AggregateOrInclusiveLayer;
 			TArray<TSharedPtr<ILayerInstance>, TInlineAllocator<32>> SubLayerInstances;
@@ -3583,9 +3690,9 @@ namespace UE::NNE::RuntimeBasic
 				NNE_RUNTIME_BASIC_TRACE_SCOPE(NNE::RuntimeBasic::Private::FAggregateOrInclusiveLayer::Evaluate);
 				check(OutputBufferSize == GetOutputSize() && InputBufferSize == GetInputSize());
 				check(OutputBufferStride >= GetOutputSize() && InputBufferStride >= GetInputSize());
+				check(Instance && Instance->GetLayerType() == GetLayerType());
 
 				FAggregateOrInclusiveLayerInstance* AggregateOrInclusiveInstance = StaticCast<FAggregateOrInclusiveLayerInstance*>(Instance);
-				check(AggregateOrInclusiveInstance);
 
 				OperatorNanCheck(InputBuffer, BatchSize, InputBufferSize, InputBufferStride);
 
@@ -3877,7 +3984,8 @@ namespace UE::NNE::RuntimeBasic
 		{
 			FSparseMixtureOfExpertsLayerInstance(const FSparseMixtureOfExpertsLayer& InSparseMixtureOfExpertsLayer);
 
-			void SetMaxBatchSize(const uint32 MaxBatchSize) override final;
+			virtual void SetMaxBatchSize(const uint32 MaxBatchSize) override final;
+			virtual ELayerType GetLayerType() const override final { return ELayerType::SparseMixtureOfExperts; }
 
 			const FSparseMixtureOfExpertsLayer& SparseMixtureOfExpertsLayer;
 			TSharedPtr<ILayerInstance> GatingInstance;
@@ -3944,9 +4052,9 @@ namespace UE::NNE::RuntimeBasic
 				NNE_RUNTIME_BASIC_TRACE_SCOPE(NNE::RuntimeBasic::Private::FSparseMixtureOfExpertsLayer::Evaluate);
 				check(OutputBufferSize == GetOutputSize() && InputBufferSize == GetInputSize());
 				check(OutputBufferStride >= GetOutputSize() && InputBufferStride >= GetInputSize());
+				check(Instance && Instance->GetLayerType() == GetLayerType());
 
 				FSparseMixtureOfExpertsLayerInstance* SparseMixtureOfExpertsInstance = StaticCast<FSparseMixtureOfExpertsLayerInstance*>(Instance);
-				check(SparseMixtureOfExpertsInstance);
 
 				OperatorNanCheck(InputBuffer, BatchSize, InputBufferSize, InputBufferStride);
 
@@ -4284,7 +4392,8 @@ namespace UE::NNE::RuntimeBasic
 		{
 			FSpreadLayerInstance(const FSpreadLayer& InSpreadLayer);
 
-			void SetMaxBatchSize(const uint32 MaxBatchSize) override final;
+			virtual void SetMaxBatchSize(const uint32 MaxBatchSize) override final;
+			virtual ELayerType GetLayerType() const override final { return ELayerType::Spread; }
 
 			const FSpreadLayer& SpreadLayer;
 			TArray<TSharedPtr<ILayerInstance>, TInlineAllocator<32>> Instances;
@@ -4350,9 +4459,9 @@ namespace UE::NNE::RuntimeBasic
 				NNE_RUNTIME_BASIC_TRACE_SCOPE(NNE::RuntimeBasic::Private::FSpreadLayer::Evaluate);
 				check(OutputBufferSize == GetOutputSize() && InputBufferSize == GetInputSize());
 				check(OutputBufferStride >= GetOutputSize() && InputBufferStride >= GetInputSize());
+				check(Instance && Instance->GetLayerType() == GetLayerType());
 
 				FSpreadLayerInstance* SpreadInstance = StaticCast<FSpreadLayerInstance*>(Instance);
-				check(SpreadInstance);
 
 				OperatorNanCheck(InputBuffer, BatchSize, InputBufferSize, InputBufferStride);
 
@@ -4553,6 +4662,7 @@ namespace UE::NNE::RuntimeBasic
 			FFiLMInstance(const FFiLMLayer& InFiLMLayer);
 
 			virtual void SetMaxBatchSize(const uint32 MaxBatchSize) override final;
+			virtual ELayerType GetLayerType() const override final { return ELayerType::FiLM; }
 
 			const FFiLMLayer& FiLMLayer;
 			TSharedPtr<ILayerInstance> PrefixInstance;
@@ -4623,9 +4733,9 @@ namespace UE::NNE::RuntimeBasic
 				check(OutputBufferStride >= GetOutputSize() && InputBufferStride >= GetInputSize());
 				check(PrefixOutputSize * 2 == ConditionOutputSize);
 				check(PostfixInputSize == PrefixOutputSize);
+				check(Instance && Instance->GetLayerType() == GetLayerType());
 
 				FFiLMInstance* FiLMInstance = StaticCast<FFiLMInstance*>(Instance);
-				check(FiLMInstance);
 
 				OperatorNanCheck(InputBuffer, BatchSize, InputBufferSize, InputBufferStride);
 
@@ -4702,6 +4812,178 @@ namespace UE::NNE::RuntimeBasic
 			ConditionBuffer.SetNumUninitialized(MaxBatchSize * FiLMLayer.ConditionOutputSize, EAllowShrinking::No);
 		}
 
+
+		//--------------------------------------------------------------------------
+
+		struct FGRUCellLayer;
+
+		struct FGRUCellInstance : public ILayerInstance
+		{
+			FGRUCellInstance(const FGRUCellLayer& InGRUCellLayer);
+
+			virtual void SetMaxBatchSize(const uint32 MaxBatchSize) override final;
+			virtual ELayerType GetLayerType() const override final { return ELayerType::GRUCell; }
+
+			const FGRUCellLayer& GRUCellLayer;
+			TSharedPtr<ILayerInstance> RememberInstance;
+			TSharedPtr<ILayerInstance> UpdateInstance;
+			TSharedPtr<ILayerInstance> ActivationInstance;
+			TArray<float> RememberGateBuffer;
+			TArray<float> UpdateGateBuffer;
+			TArray<float> ActivationGateBuffer;
+			TArray<float> InputMaskBuffer;
+		};
+
+		struct FGRUCellLayer : public ILayer
+		{
+			virtual TSharedPtr<ILayerInstance> MakeInstance() const { return MakeShared<FGRUCellInstance>(*this); };
+			virtual ELayerType GetLayerType() const override final { return ELayerType::GRUCell; }
+			virtual uint32 GetInputSize() const override final { return InputSize + OutputSize; }
+			virtual uint32 GetOutputSize() const override final { return OutputSize + OutputSize; }
+
+			virtual void SerializationSize(uint64& InOutOffset) const override final
+			{
+				Serialization::Size(InOutOffset, InputSize);
+				Serialization::Size(InOutOffset, OutputSize);
+				Serialization::Size(InOutOffset, RememberLayer);
+				Serialization::Size(InOutOffset, UpdateLayer);
+				Serialization::Size(InOutOffset, ActivationLayer);
+			}
+
+			virtual void SerializationLoad(uint64& InOutOffset, TConstArrayView<uint8> Data) override final
+			{
+				Serialization::Load(InOutOffset, InputSize, Data);
+				Serialization::Load(InOutOffset, OutputSize, Data);
+				Serialization::Load(InOutOffset, RememberLayer, Data);
+				Serialization::Load(InOutOffset, UpdateLayer, Data);
+				Serialization::Load(InOutOffset, ActivationLayer, Data);
+			}
+
+			virtual void SerializationSave(uint64& InOutOffset, TArrayView<uint8> Data) const override final
+			{
+				Serialization::Save(InOutOffset, InputSize, Data);
+				Serialization::Save(InOutOffset, OutputSize, Data);
+				Serialization::Save(InOutOffset, RememberLayer, Data);
+				Serialization::Save(InOutOffset, UpdateLayer, Data);
+				Serialization::Save(InOutOffset, ActivationLayer, Data);
+			}
+
+			virtual void Evaluate(
+				ILayerInstance* Instance,
+				float* OutputBuffer,
+				const float* InputBuffer,
+				const uint32 BatchSize,
+				const uint32 OutputBufferSize,
+				const uint32 InputBufferSize,
+				const uint32 OutputBufferStride,
+				const uint32 InputBufferStride) override final
+			{
+				NNE_RUNTIME_BASIC_TRACE_SCOPE(NNE::RuntimeBasic::Private::FGRUCellLayer::Evaluate);
+				check(OutputBufferSize == GetOutputSize() && InputBufferSize == GetInputSize());
+				check(OutputBufferStride >= GetOutputSize() && InputBufferStride >= GetInputSize());
+				check(Instance && Instance->GetLayerType() == GetLayerType());
+
+				FGRUCellInstance* GRUCellInstance = StaticCast<FGRUCellInstance*>(Instance);
+
+				OperatorNanCheck(InputBuffer, BatchSize, InputBufferSize, InputBufferStride);
+
+				// Remember Gate
+
+				RememberLayer->Evaluate(
+					GRUCellInstance->RememberInstance.Get(),
+					GRUCellInstance->RememberGateBuffer.GetData(),
+					InputBuffer,
+					BatchSize,
+					OutputSize,
+					InputBufferSize,
+					OutputSize,
+					InputBufferStride);
+
+				// Update Gate
+
+				UpdateLayer->Evaluate(
+					GRUCellInstance->UpdateInstance.Get(),
+					GRUCellInstance->UpdateGateBuffer.GetData(),
+					InputBuffer,
+					BatchSize,
+					OutputSize,
+					InputBufferSize,
+					OutputSize,
+					InputBufferStride);
+
+				// Mask Input
+
+				OperatorGRUMaskInput(
+					GRUCellInstance->InputMaskBuffer.GetData(),
+					InputBuffer,
+					GRUCellInstance->RememberGateBuffer.GetData(),
+					BatchSize,
+					OutputSize,
+					InputSize,
+					InputSize + OutputSize,
+					InputBufferStride,
+					OutputSize);
+
+				// Activation Gate
+
+				ActivationLayer->Evaluate(
+					GRUCellInstance->ActivationInstance.Get(),
+					GRUCellInstance->ActivationGateBuffer.GetData(),
+					GRUCellInstance->InputMaskBuffer.GetData(),
+					BatchSize,
+					OutputSize,
+					InputSize + OutputSize,
+					OutputSize,
+					InputSize + OutputSize);
+
+				// Compute Output
+
+				OperatorGRUCellUpdateOutput(
+					OutputBuffer,
+					InputBuffer,
+					GRUCellInstance->UpdateGateBuffer.GetData(),
+					GRUCellInstance->ActivationGateBuffer.GetData(),
+					BatchSize,
+					OutputSize,
+					InputSize,
+					OutputBufferStride,
+					InputBufferStride,
+					OutputSize,
+					OutputSize);
+
+				OperatorNanCheck(OutputBuffer, BatchSize, OutputBufferSize, OutputBufferStride);
+			}
+
+			uint32 InputSize = 0;
+			uint32 OutputSize = 0;
+			TSharedPtr<ILayer> RememberLayer;
+			TSharedPtr<ILayer> UpdateLayer;
+			TSharedPtr<ILayer> ActivationLayer;
+		};
+
+		FGRUCellInstance::FGRUCellInstance(const FGRUCellLayer& InGRUCellLayer)
+			: GRUCellLayer(InGRUCellLayer)
+		{
+			RememberInstance = GRUCellLayer.RememberLayer->MakeInstance();
+			UpdateInstance = GRUCellLayer.UpdateLayer->MakeInstance();
+			ActivationInstance = GRUCellLayer.ActivationLayer->MakeInstance();
+		}
+
+		void FGRUCellInstance::SetMaxBatchSize(const uint32 MaxBatchSize)
+		{
+			NNE_RUNTIME_BASIC_TRACE_SCOPE(NNE::RuntimeBasic::Private::FGRUCellInstance::SetMaxBatchSize);
+
+			if (RememberInstance) { RememberInstance->SetMaxBatchSize(MaxBatchSize); }
+			if (UpdateInstance) { UpdateInstance->SetMaxBatchSize(MaxBatchSize); }
+			if (ActivationInstance) { ActivationInstance->SetMaxBatchSize(MaxBatchSize); }
+
+			RememberGateBuffer.SetNumUninitialized(MaxBatchSize * GRUCellLayer.OutputSize, EAllowShrinking::No);
+			UpdateGateBuffer.SetNumUninitialized(MaxBatchSize * GRUCellLayer.OutputSize, EAllowShrinking::No);
+			ActivationGateBuffer.SetNumUninitialized(MaxBatchSize * GRUCellLayer.OutputSize, EAllowShrinking::No);
+			InputMaskBuffer.SetNumUninitialized(MaxBatchSize * (GRUCellLayer.InputSize + GRUCellLayer.OutputSize), EAllowShrinking::No);
+		}
+
+
 		//--------------------------------------------------------------------------
 		// Layer Serialization
 		//--------------------------------------------------------------------------
@@ -4757,6 +5039,7 @@ namespace UE::NNE::RuntimeBasic
 					case ELayerType::Slice: OutLayer = MakeShared<FSliceLayer>(); break;
 					case ELayerType::Residual: OutLayer = MakeShared<FResidualLayer>(); break;
 					case ELayerType::FiLM: OutLayer = MakeShared<FFiLMLayer>(); break;
+					case ELayerType::GRUCell: OutLayer = MakeShared<FGRUCellLayer>(); break;
 					default: checkf(false, TEXT("Unknown Layer Id %i"), LayerTypeId);
 					}
 				}
@@ -5635,6 +5918,44 @@ namespace UE::NNE::RuntimeBasic
 			MakeLinearLayer(MemoryNum, OutputNum, LinearLayerSettings));
 	}
 
+
+	FModelBuilderElement FModelBuilder::MakeGRUCell(
+		const uint32 InputNum,
+		const uint32 OutputNum,
+		const FModelBuilderElement& RememberLayer,
+		const FModelBuilderElement& UpdateLayer,
+		const FModelBuilderElement& ActivationLayer)
+	{
+		check(RememberLayer.GetInputSize() == InputNum + OutputNum);
+		check(RememberLayer.GetOutputSize() == OutputNum);
+		check(UpdateLayer.GetInputSize() == InputNum + OutputNum);
+		check(UpdateLayer.GetOutputSize() == OutputNum);
+		check(ActivationLayer.GetInputSize() == InputNum + OutputNum);
+		check(ActivationLayer.GetOutputSize() == OutputNum);
+
+		TSharedPtr<Private::FGRUCellLayer> CellLayer = MakeShared<Private::FGRUCellLayer>();
+		CellLayer->InputSize = InputNum;
+		CellLayer->OutputSize = OutputNum;
+		CellLayer->RememberLayer = RememberLayer.Layer;
+		CellLayer->UpdateLayer = UpdateLayer.Layer;
+		CellLayer->ActivationLayer = ActivationLayer.Layer;
+
+		return StaticCastSharedPtr<Private::ILayer>(CellLayer);
+	}
+
+	FModelBuilderElement FModelBuilder::MakeGRUCellLayer(
+		const uint32 InputNum,
+		const uint32 OutputNum,
+		const FLinearLayerSettings& LinearLayerSettings)
+	{
+		return MakeGRUCell(
+			InputNum,
+			OutputNum,
+			MakeLinearLayer(InputNum + OutputNum, OutputNum, LinearLayerSettings),
+			MakeLinearLayer(InputNum + OutputNum, OutputNum, LinearLayerSettings),
+			MakeLinearLayer(InputNum + OutputNum, OutputNum, LinearLayerSettings));
+	}
+
 	FModelBuilderElement FModelBuilder::MakeMemoryBackbone(
 		const uint32 MemoryNum,
 		const FModelBuilderElement& Prefix,
diff --git a/Engine/Plugins/Experimental/NNERuntimeBasicCpu/Source/NNERuntimeBasicCpu/Public/NNERuntimeBasicCpuBuilder.h b/Engine/Plugins/Experimental/NNERuntimeBasicCpu/Source/NNERuntimeBasicCpu/Public/NNERuntimeBasicCpuBuilder.h
--- a/Engine/Plugins/Experimental/NNERuntimeBasicCpu/Source/NNERuntimeBasicCpu/Public/NNERuntimeBasicCpuBuilder.h
+++ b/Engine/Plugins/Experimental/NNERuntimeBasicCpu/Source/NNERuntimeBasicCpu/Public/NNERuntimeBasicCpuBuilder.h
@@ -404,6 +404,34 @@ namespace UE::NNE::RuntimeBasic
 			const uint32 MemoryNum,
 			const FLinearLayerSettings& LinearLayerSettings = FLinearLayerSettings());
 
+		/**
+		 * Make a new GRU Cell layer.
+		 *
+		 * @param InputNum					Number of inputs to the model.
+		 * @param OutputNum					Number of outputs from the model. Will also be the memory size.
+		 * @param RememberLayer				Layer used for the Remember Gate.
+		 * @param UpdateLayer				Layer used for the Update Gate.
+		 * @param ActivationLayer			Layer used for producing the Activation values.
+		 */
+		UE_API FModelBuilderElement MakeGRUCell(
+			const uint32 InputNum,
+			const uint32 OutputNum,
+			const FModelBuilderElement& RememberLayer,
+			const FModelBuilderElement& UpdateLayer,
+			const FModelBuilderElement& ActivationLayer);
+
+		/**
+		 * Make a new GRU Cell layer with the given linear layer settings.
+		 *
+		 * @param InputNum				Number of inputs to the model.
+		 * @param OutputNum				Number of outputs from the model. Will also be the memory size.
+		 * @param LinearLayerSettings	Settings for the linear layers in the Memory Cell.
+		 */
+		UE_API FModelBuilderElement MakeGRUCellLayer(
+			const uint32 InputNum,
+			const uint32 OutputNum,
+			const FLinearLayerSettings& LinearLayerSettings = FLinearLayerSettings());
+
 		/**
 		 * Make a new Memory Backbone layer.
 		 *
