[NNERuntimeBasicCPU] Add pragma once to avoid compile warning on modulo

diff --git a/Engine/Plugins/Experimental/NNERuntimeBasicCpu/Source/NNERuntimeBasicCpu/Private/NNERuntimeBasicCpu.ispc b/Engine/Plugins/Experimental/NNERuntimeBasicCpu/Source/NNERuntimeBasicCpu/Private/NNERuntimeBasicCpu.ispc
--- a/Engine/Plugins/Experimental/NNERuntimeBasicCpu/Source/NNERuntimeBasicCpu/Private/NNERuntimeBasicCpu.ispc
+++ b/Engine/Plugins/Experimental/NNERuntimeBasicCpu/Source/NNERuntimeBasicCpu/Private/NNERuntimeBasicCpu.ispc
@@ -743,6 +743,7 @@ export void NNERuntimeBasicCPUOperatorConv1d(
 						for (uniform int KernelIdx = 0; KernelIdx < KernelSize; KernelIdx++)
 						{
 							int Idx = OutputIdx + KernelIdx - Padding;
+							#pragma ignore warning
 							Idx = Idx % InputNum;
 							if (Idx < 0) {Idx += InputNum;}
 							Sum += InChannelsItr[Idx] * WeightsInChannelsItr[KernelIdx];
