[LearningAgents] [NNE] Conv2d

diff --git a/Engine/Plugins/Experimental/LearningAgents/Source/Learning/Private/LearningObservation.cpp b/Engine/Plugins/Experimental/LearningAgents/Source/Learning/Private/LearningObservation.cpp
--- a/Engine/Plugins/Experimental/LearningAgents/Source/Learning/Private/LearningObservation.cpp
+++ b/Engine/Plugins/Experimental/LearningAgents/Source/Learning/Private/LearningObservation.cpp
@@ -325,6 +325,36 @@ namespace UE::Learning::Observation
 		return { Index, Generation };
 	}
 
+	FSchemaElement FSchema::CreateConv2d(const FSchemaConv2dParameters Parameters, const FName Tag)
+	{
+		FConv2dData ElementData;
+		ElementData.ElementIndex = SubElementObjects.Num();
+		ElementData.InputHeight = Parameters.InputHeight;
+		ElementData.InputWidth = Parameters.InputWidth;
+		ElementData.InChannels = Parameters.InChannels;
+		ElementData.OutChannels = Parameters.OutChannels;
+		ElementData.KernelSize = Parameters.KernelSize;
+		ElementData.Stride = Parameters.Stride;
+		ElementData.Padding = Parameters.Padding;
+		ElementData.PaddingMode = Parameters.PaddingMode;
+		ElementData.ActivationFunction = Parameters.ActivationFunction;
+
+		SubElementNames.Add(NAME_None);
+		SubElementObjects.Add(Parameters.Element);
+
+		const int32 Index = Types.Add(EType::Conv2d);
+		Tags.Add(Tag);
+
+		ObservationVectorSizes.Add(GetObservationVectorSize(Parameters.Element));
+		int32 OutputHeight = (Parameters.InputHeight + 2 * Parameters.Padding - Parameters.KernelSize) / Parameters.Stride + 1;
+		int32 OutputWidth = (Parameters.InputWidth + 2 * Parameters.Padding - Parameters.KernelSize) / Parameters.Stride + 1;
+		EncodedVectorSizes.Add(OutputHeight * OutputWidth * Parameters.OutChannels);
+
+		TypeDataIndices.Add(Conv2dData.Add(ElementData));
+
+		return { Index, Generation };
+	}
+
 	bool FSchema::IsValid(const FSchemaElement Element) const
 	{
 		return Element.Generation == Generation && Element.Index != INDEX_NONE;
@@ -497,6 +527,25 @@ namespace UE::Learning::Observation
 		return Parameters;
 	}
 
+	FSchemaConv2dParameters FSchema::GetConv2d(const FSchemaElement Element) const
+	{
+		check(IsValid(Element) && GetType(Element) == EType::Conv2d);
+		const FConv2dData& ElementData = Conv2dData[TypeDataIndices[Element.Index]];
+
+		FSchemaConv2dParameters Parameters;
+		Parameters.Element = SubElementObjects[ElementData.ElementIndex];
+		Parameters.InputHeight = ElementData.InputHeight;
+		Parameters.InputWidth = ElementData.InputWidth;
+		Parameters.InChannels = ElementData.InChannels;
+		Parameters.OutChannels = ElementData.OutChannels;
+		Parameters.KernelSize = ElementData.KernelSize;
+		Parameters.Stride = ElementData.Stride;
+		Parameters.Padding = ElementData.Padding;
+		Parameters.PaddingMode = ElementData.PaddingMode;
+		Parameters.ActivationFunction = ElementData.ActivationFunction;
+		return Parameters;
+	}
+
 	uint32 FSchema::GetGeneration() const
 	{
 		return Generation;
@@ -839,6 +888,28 @@ namespace UE::Learning::Observation
 		return { Index, Generation };
 	}
 
+	FObjectElement FObject::CreateConv2d(const FObjectConv2dParameters Parameters, const FName Tag)
+	{
+		check(IsValid(Parameters.Element));
+
+		const int32 Index = Types.Add(EType::Conv2d);
+		Tags.Add(Tag);
+
+		ContinuousDataOffsets.Add(ContinuousValues.Num());
+		ContinuousDataNums.Add(0);
+
+		DiscreteDataOffsets.Add(DiscreteValues.Num());
+		DiscreteDataNums.Add(0);
+
+		SubElementDataOffsets.Add(SubElementObjects.Num());
+		SubElementDataNums.Add(1);
+
+		SubElementNames.Add(NAME_Name);
+		SubElementObjects.Add(Parameters.Element);
+
+		return { Index, Generation };
+	}
+
 	bool FObject::IsValid(const FObjectElement Element) const
 	{
 		return Element.Generation == Generation && Element.Index != INDEX_NONE;
@@ -967,6 +1038,15 @@ namespace UE::Learning::Observation
 		return Parameters;
 	}
 
+	FObjectConv2dParameters FObject::GetConv2d(const FObjectElement Element) const
+	{
+		check(IsValid(Element) && GetType(Element) == EType::Conv2d);
+
+		FObjectConv2dParameters Parameters;
+		Parameters.Element = SubElementObjects[SubElementDataOffsets[Element.Index]];
+		return Parameters;
+	}
+
 	uint32 FObject::GetGeneration() const
 	{
 		return Generation;
@@ -1183,6 +1263,21 @@ namespace UE::Learning::Observation
 				GetSchemaObjectsCompatibilityHash(Schema, Parameters.Element, Salt) });
 		}
 
+		case EType::Conv2d:
+		{
+			const FSchemaConv2dParameters Parameters = Schema.GetConv2d(SchemaElement);
+			return Private::HashCombine({
+				Hash,
+				Private::HashInt(Parameters.InputHeight),
+				Private::HashInt(Parameters.InputWidth),
+				Private::HashInt(Parameters.InChannels),
+				Private::HashInt(Parameters.OutChannels),
+				Private::HashInt(Parameters.KernelSize),
+				Private::HashInt(Parameters.Stride),
+				Private::HashInt(Parameters.Padding),
+				GetSchemaObjectsCompatibilityHash(Schema, Parameters.Element, Salt) });
+		}
+
 		default:
 		{
 			checkNoEntry();
@@ -1341,6 +1436,22 @@ namespace UE::Learning::Observation
 				&& AreSchemaObjectsCompatible(SchemaA, ParametersA.Element, SchemaB, ParametersB.Element);
 		}
 
+
+		case EType::Conv2d:
+		{
+			const FSchemaConv2dParameters ParametersA = SchemaA.GetConv2d(SchemaElementA);
+			const FSchemaConv2dParameters ParametersB = SchemaB.GetConv2d(SchemaElementB);
+
+			return (ParametersA.InputHeight == ParametersB.InputHeight)
+				&& (ParametersA.InputWidth == ParametersB.InputWidth)
+				&& (ParametersA.InChannels == ParametersB.InChannels)
+				&& (ParametersA.OutChannels == ParametersB.OutChannels)
+				&& (ParametersA.KernelSize == ParametersB.KernelSize)
+				&& (ParametersA.Stride == ParametersB.Stride)
+				&& (ParametersA.Padding == ParametersB.Padding)
+				&& AreSchemaObjectsCompatible(SchemaA, ParametersA.Element, SchemaB, ParametersB.Element);
+		}
+
 		default:
 		{
 			checkNoEntry();
@@ -1625,18 +1736,60 @@ namespace UE::Learning::Observation
 			default: checkNoEntry();
 			}
 
-			OutElement = Builder.MakeConv1d(
+			NNE::RuntimeBasic::FModelBuilderElement Conv1dLayer = Builder.MakeConv1d(
 				Parameters.InputLength,
 				Parameters.InChannels,
 				Parameters.OutChannels,
 				Parameters.KernelSize,
 				Parameters.Padding,
 				Private::GetNNEPaddingMode(Parameters.PaddingMode),
-				Private::GetNNEActivationFunction(Parameters.ActivationFunction),
 				Builder.MakeInitialWeights(Parameters.InChannels * Parameters.KernelSize, Parameters.OutChannels, WeightInitializationSettings),
 				Builder.MakeInitialBiases(Parameters.OutChannels, WeightInitializationSettings),
 				BuilderSubElement);
 			
+			OutElement = Builder.MakeSequence({
+				Conv1dLayer,
+				Builder.MakeActivation(Conv1dLayer.GetOutputSize(), Private::GetNNEActivationFunction(Parameters.ActivationFunction))
+			});
+
+			break;
+		}
+
+		case EType::Conv2d:
+		{
+			const FSchemaConv2dParameters Parameters = Schema.GetConv2d(SchemaElement);
+			const int32 SubElementEncodedSize = Schema.GetEncodedVectorSize(Parameters.Element);
+			NNE::RuntimeBasic::FModelBuilderElement BuilderSubElement;
+			MakeEncoderNetworkModelBuilderElementFromSchema(BuilderSubElement, Builder, Schema, Parameters.Element, NetworkSettings);
+
+			NNE::RuntimeBasic::FModelBuilder::FWeightInitializationSettings WeightInitializationSettings;
+			switch (NetworkSettings.WeightInitialization)
+			{
+			case EWeightInitialization::KaimingGaussian: WeightInitializationSettings.Type =
+				NNE::RuntimeBasic::FModelBuilder::EWeightInitializationType::KaimingGaussian; break;
+			case EWeightInitialization::KaimingUniform: WeightInitializationSettings.Type =
+				NNE::RuntimeBasic::FModelBuilder::EWeightInitializationType::KaimingUniform; break;
+			default: checkNoEntry();
+			}
+
+			NNE::RuntimeBasic::FModelBuilderElement Conv2dLayer = Builder.MakeConv2d(
+				Parameters.InputHeight,
+				Parameters.InputWidth,
+				Parameters.InChannels,
+				Parameters.OutChannels,
+				Parameters.KernelSize,
+				Parameters.Stride,
+				Parameters.Padding,
+				Private::GetNNEPaddingMode(Parameters.PaddingMode),
+				Builder.MakeInitialWeights(Parameters.InChannels * Parameters.KernelSize * Parameters.KernelSize, Parameters.OutChannels, WeightInitializationSettings),
+				Builder.MakeInitialBiases(Parameters.OutChannels, WeightInitializationSettings),
+				BuilderSubElement);
+
+			OutElement = Builder.MakeSequence({
+				Conv2dLayer,
+				Builder.MakeActivation(Conv2dLayer.GetOutputSize(), Private::GetNNEActivationFunction(Parameters.ActivationFunction))
+			});
+
 			break;
 		}
 
@@ -1983,6 +2136,21 @@ namespace UE::Learning::Observation
 			return;
 		}
 
+		case EType::Conv2d:
+		{
+			const FSchemaConv2dParameters SchemaParameters = Schema.GetConv2d(SchemaElement);
+			const FObjectConv2dParameters ObjectParameters = Object.GetConv2d(ObjectElement);
+
+			SetVectorFromObject(
+				OutObservationVector,
+				Schema,
+				SchemaParameters.Element,
+				Object,
+				ObjectParameters.Element);
+
+			return;
+		}
+
 		default:
 		{
 			checkNoEntry();
@@ -2332,6 +2500,22 @@ namespace UE::Learning::Observation
 			return;
 		}
 
+		case EType::Conv2d:
+		{
+			const FSchemaConv2dParameters Parameters = Schema.GetConv2d(SchemaElement);
+
+			FObjectElement SubElement;
+			GetObjectFromVector(
+				OutObject,
+				SubElement,
+				Schema,
+				Parameters.Element,
+				ObservationVector);
+
+			OutObjectElement = OutObject.CreateConv2d({ SubElement }, SchemaElementTag);
+			return;
+		}
+
 		default:
 		{
 			checkNoEntry();
@@ -2551,6 +2735,18 @@ namespace UE::Learning::Observation
 			return;
 		}
 
+		case EType::Conv2d:
+		{
+			AddGaussianNoiseToVector(
+				InOutRandomState,
+				InOutObservationVector,
+				Schema,
+				Schema.GetConv2d(SchemaElement).Element,
+				NoiseScale);
+
+			return;
+		}
+
 		default:
 		{
 			checkNoEntry();
diff --git a/Engine/Plugins/Experimental/LearningAgents/Source/Learning/Public/LearningObservation.h b/Engine/Plugins/Experimental/LearningAgents/Source/Learning/Public/LearningObservation.h
--- a/Engine/Plugins/Experimental/LearningAgents/Source/Learning/Public/LearningObservation.h
+++ b/Engine/Plugins/Experimental/LearningAgents/Source/Learning/Public/LearningObservation.h
@@ -61,6 +61,9 @@ namespace UE::Learning::Observation
 
 		// 1d convolution of another observation 
 		Conv1d = 12,
+
+		// 2d convolution of another observation 
+		Conv2d = 13,
 	};
 
 	/**
@@ -207,19 +210,52 @@ namespace UE::Learning::Observation
 		FSchemaElement Element;
 
 		// Length of the 1d convolution input.
-		int32 InputLength;
+		int32 InputLength = 0;
 
 		// Number of input channels.
-		int32 InChannels;
+		int32 InChannels = 0;
 
 		// Number of output channels after convolution.
-		int32 OutChannels;
+		int32 OutChannels = 0;
 
 		// Size of the convolving kernel.
-		int32 KernelSize;
+		int32 KernelSize = 3;
 
 		// The amount of padding to add to both sides of the input.
-		int32 Padding;
+		int32 Padding = 1;
+
+		// Zeros or Circular.
+		EPaddingMode PaddingMode = EPaddingMode::Zeros;
+
+		// The activation function to use for the convolution layer.
+		EEncodingActivationFunction ActivationFunction = EEncodingActivationFunction::ReLU;
+	};
+
+	struct FSchemaConv2dParameters
+	{
+		// The sub-element.
+		FSchemaElement Element;
+
+		// Height of the 2d convolution input.
+		int32 InputHeight = 0;
+		
+		// Width of the 2d convolution input.
+		int32 InputWidth = 0;
+
+		// Number of input channels.
+		int32 InChannels = 0;
+
+		// Number of output channels after convolution.
+		int32 OutChannels = 0;
+
+		// Size of the convolving kernel.
+		int32 KernelSize = 3;
+
+		// Size of the convolution stride.
+		int32 Stride = 1;
+
+		// The amount of padding to add to both sides of the input.
+		int32 Padding = 1;
 
 		// Zeros or Circular.
 		EPaddingMode PaddingMode = EPaddingMode::Zeros;
@@ -250,6 +286,7 @@ namespace UE::Learning::Observation
 		UE_API FSchemaElement CreateSet(const FSchemaSetParameters Parameters, const FName Tag = NAME_None);
 		UE_API FSchemaElement CreateEncoding(const FSchemaEncodingParameters Parameters, const FName Tag = NAME_None);
 		UE_API FSchemaElement CreateConv1d(const FSchemaConv1dParameters Parameters, const FName Tag = NAME_None);
+		UE_API FSchemaElement CreateConv2d(const FSchemaConv2dParameters Parameters, const FName Tag = NAME_None);
 
 		UE_API FSchemaContinuousParameters GetContinuous(const FSchemaElement Element) const;
 		UE_API FSchemaDiscreteExclusiveParameters GetDiscreteExclusive(const FSchemaElement Element) const;
@@ -263,6 +300,7 @@ namespace UE::Learning::Observation
 		UE_API FSchemaSetParameters GetSet(const FSchemaElement Element) const;
 		UE_API FSchemaEncodingParameters GetEncoding(const FSchemaElement Element) const;
 		UE_API FSchemaConv1dParameters GetConv1d(const FSchemaElement Element) const;
+		UE_API FSchemaConv2dParameters GetConv2d(const FSchemaElement Element) const;
 
 		// Checks if the given element is valid
 		UE_API bool IsValid(const FSchemaElement Element) const;
@@ -378,6 +416,20 @@ namespace UE::Learning::Observation
 			EEncodingActivationFunction ActivationFunction = EEncodingActivationFunction::ReLU;
 		};
 
+		struct FConv2dData
+		{
+			int32 ElementIndex = INDEX_NONE;
+			int32 InputHeight = INDEX_NONE;
+			int32 InputWidth = INDEX_NONE;
+			int32 InChannels = INDEX_NONE;
+			int32 OutChannels = INDEX_NONE;
+			int32 KernelSize = INDEX_NONE;
+			int32 Stride = INDEX_NONE;
+			int32 Padding = INDEX_NONE;
+			EPaddingMode PaddingMode = EPaddingMode::Zeros;
+			EEncodingActivationFunction ActivationFunction = EEncodingActivationFunction::ReLU;
+		};
+
 		uint32 Generation = 0;
 
 		/** These have entries for each Schema Element */
@@ -400,6 +452,7 @@ namespace UE::Learning::Observation
 		TArray<FSetData> SetData;
 		TArray<FEncodingData> EncodingData;
 		TArray<FConv1dData> Conv1dData;
+		TArray<FConv2dData> Conv2dData;
 
 		/** This is an array of all the SubElements and their names, referenced by other elements */
 		TArray<FName> SubElementNames;
@@ -495,6 +548,12 @@ namespace UE::Learning::Observation
 		FObjectElement Element;
 	};
 
+	struct FObjectConv2dParameters
+	{
+		// Encoded sub-element.
+		FObjectElement Element;
+	};
+
 	/**
 	 * Observation Object
 	 *
@@ -517,6 +576,7 @@ namespace UE::Learning::Observation
 		UE_API FObjectElement CreateSet(const FObjectSetParameters Parameters, const FName Tag = NAME_None);
 		UE_API FObjectElement CreateEncoding(const FObjectEncodingParameters Parameters, const FName Tag = NAME_None);
 		UE_API FObjectElement CreateConv1d(const FObjectConv1dParameters Parameters, const FName Tag = NAME_None);
+		UE_API FObjectElement CreateConv2d(const FObjectConv2dParameters Parameters, const FName Tag = NAME_None);
 
 		UE_API FObjectContinuousParameters GetContinuous(const FObjectElement Element) const;
 		UE_API FObjectDiscreteExclusiveParameters GetDiscreteExclusive(const FObjectElement Element) const;
@@ -530,6 +590,7 @@ namespace UE::Learning::Observation
 		UE_API FObjectSetParameters GetSet(const FObjectElement Element) const;
 		UE_API FObjectEncodingParameters GetEncoding(const FObjectElement Element) const;
 		UE_API FObjectConv1dParameters GetConv1d(const FObjectElement Element) const;
+		UE_API FObjectConv2dParameters GetConv2d(const FObjectElement Element) const;
 
 		// Checks if the given element is valid
 		UE_API bool IsValid(const FObjectElement Element) const;
diff --git a/Engine/Plugins/Experimental/LearningAgents/Source/LearningAgents/Private/LearningAgentsObservations.cpp b/Engine/Plugins/Experimental/LearningAgents/Source/LearningAgents/Private/LearningAgentsObservations.cpp
--- a/Engine/Plugins/Experimental/LearningAgents/Source/LearningAgents/Private/LearningAgentsObservations.cpp
+++ b/Engine/Plugins/Experimental/LearningAgents/Source/LearningAgents/Private/LearningAgentsObservations.cpp
@@ -429,6 +429,18 @@ namespace UE::Learning::Agents::Observation::Private
 				ObjectParameters.Element);
 		}
 
+		case Learning::Observation::EType::Conv2d:
+		{
+			const Learning::Observation::FSchemaConv2dParameters SchemaParameters = Schema.GetConv2d(SchemaElement);
+			const Learning::Observation::FObjectConv2dParameters ObjectParameters = Object.GetConv2d(ObjectElement);
+
+			return ValidateObservationObjectMatchesSchema(
+				Schema,
+				SchemaParameters.Element,
+				Object,
+				ObjectParameters.Element);
+		}
+
 		default:
 		{
 			checkNoEntry();
@@ -1227,7 +1239,7 @@ FLearningAgentsObservationSchemaElement ULearningAgentsObservations::SpecifyBool
 	return SpecifyContinuousObservation(Schema, 1, 1.0f, Tag);
 }
 
-FLearningAgentsObservationSchemaElement ULearningAgentsObservations::SpecifyConv1dObservation(ULearningAgentsObservationSchema* Schema, const FLearningAgentsObservationSchemaElement Element, const int32 InputLength, const int32 InChannels, const int32 OutChannels, const int32 KernelSize, const int32 Padding, const ELearningAgentsPaddingMode PaddingMode, const ELearningAgentsActivationFunction ActivationFunction, const FName Tag /*= TEXT("Conv1dObservation")*/)
+FLearningAgentsObservationSchemaElement ULearningAgentsObservations::SpecifyConv1dObservation(ULearningAgentsObservationSchema* Schema, const FLearningAgentsObservationSchemaElement Element, const int32 InputLength, const int32 InChannels, const int32 OutChannels, const int32 KernelSize, const int32 Padding, const ELearningAgentsPaddingMode PaddingMode, const ELearningAgentsActivationFunction ActivationFunction, const FName Tag)
 {
 	if (!Schema)
 	{
@@ -1284,6 +1296,70 @@ FLearningAgentsObservationSchemaElement ULearningAgentsObservations::SpecifyConv
 	Tag) };
 }
 
+FLearningAgentsObservationSchemaElement ULearningAgentsObservations::SpecifyConv2dObservation(ULearningAgentsObservationSchema* Schema, const FLearningAgentsObservationSchemaElement Element, const int32 InputHeight, const int32 InputWidth, const int32 InChannels, const int32 OutChannels, const int32 KernelSize, const int32 Stride, const int32 Padding, const ELearningAgentsActivationFunction ActivationFunction, const FName Tag)
+{
+	if (!Schema)
+	{
+		UE_LOG(LogLearning, Error, TEXT("SpecifyConv2dObservation: Schema is nullptr."));
+		return FLearningAgentsObservationSchemaElement();
+	}
+	if (!Schema->ObservationSchema.IsValid(Element.SchemaElement))
+	{
+		UE_LOG(LogLearning, Error, TEXT("SpecifyConv2dObservation: Invalid Observation Object."));
+		return FLearningAgentsObservationSchemaElement();
+	}
+	if (InputHeight < 1)
+	{
+		UE_LOG(LogLearning, Error, TEXT("SpecifyConv2dObservation: Invalid Observation InputHeight '%i' - must be greater than zero."), InputHeight);
+		return FLearningAgentsObservationSchemaElement();
+	}
+	if (InputWidth < 1)
+	{
+		UE_LOG(LogLearning, Error, TEXT("SpecifyConv2dObservation: Invalid Observation InputWidth '%i' - must be greater than zero."), InputWidth);
+		return FLearningAgentsObservationSchemaElement();
+	}
+	if (InChannels < 1)
+	{
+		UE_LOG(LogLearning, Error, TEXT("SpecifyConv2dObservation: Invalid Observation InChannels '%i' - must be greater than zero."), InChannels);
+		return FLearningAgentsObservationSchemaElement();
+	}
+	if (OutChannels < 1)
+	{
+		UE_LOG(LogLearning, Error, TEXT("SpecifyConv2dObservation: Invalid Observation OutChannels '%i' - must be greater than zero."), OutChannels);
+		return FLearningAgentsObservationSchemaElement();
+	}
+	if (KernelSize < 1)
+	{
+		UE_LOG(LogLearning, Error, TEXT("SpecifyConv2dObservation: Invalid Observation KernelSize '%i' - must be greater than zero."), KernelSize);
+		return FLearningAgentsObservationSchemaElement();
+	}
+	if (Stride < 1)
+	{
+		UE_LOG(LogLearning, Error, TEXT("SpecifyConv2dObservation: Invalid Observation Stride '%i' - must be greater than zero."), Stride);
+		return FLearningAgentsObservationSchemaElement();
+	}
+	if (Padding < 0)
+	{
+		UE_LOG(LogLearning, Error, TEXT("SpecifyConv2dObservation: Invalid Observation Padding '%i' - must be greater than zero."), Padding);
+		return FLearningAgentsObservationSchemaElement();
+	}
+
+	return { Schema->ObservationSchema.CreateConv2d(
+	{
+		Element.SchemaElement,
+		InputHeight,
+		InputWidth,
+		InChannels,
+		OutChannels,
+		KernelSize,
+		Stride,
+		Padding,
+		UE::Learning::Agents::Observation::Private::GetPaddingMode(ELearningAgentsPaddingMode::Zeros),
+		UE::Learning::Agents::Observation::Private::GetEncodingActivationFunction(ActivationFunction)
+	},
+	Tag) };
+}
+
 FLearningAgentsObservationSchemaElement ULearningAgentsObservations::SpecifyFloatObservation(ULearningAgentsObservationSchema* Schema, const float FloatScale, const FName Tag)
 {
 	return SpecifyContinuousObservation(Schema, 1, FloatScale, Tag);
@@ -2197,7 +2273,7 @@ FLearningAgentsObservationObjectElement ULearningAgentsObservations::MakeBoolObs
 	return MakeContinuousObservationFromArrayView(Object, { EncodedBool }, Tag);
 }
 
-FLearningAgentsObservationObjectElement ULearningAgentsObservations::MakeConv1dObservation(ULearningAgentsObservationObject* Object, const FLearningAgentsObservationObjectElement Element, const FName Tag /*= TEXT("Conv1dObservation")*/)
+FLearningAgentsObservationObjectElement ULearningAgentsObservations::MakeConv1dObservation(ULearningAgentsObservationObject* Object, const FLearningAgentsObservationObjectElement Element, const FName Tag)
 {
 	if (!Object)
 	{
@@ -2214,6 +2290,23 @@ FLearningAgentsObservationObjectElement ULearningAgentsObservations::MakeConv1dO
 	return { Object->ObservationObject.CreateConv1d({ Element.ObjectElement }, Tag) };
 }
 
+FLearningAgentsObservationObjectElement ULearningAgentsObservations::MakeConv2dObservation(ULearningAgentsObservationObject* Object, const FLearningAgentsObservationObjectElement Element, const FName Tag)
+{
+	if (!Object)
+	{
+		UE_LOG(LogLearning, Error, TEXT("MakeConv2dObservation: Object is nullptr."));
+		return FLearningAgentsObservationObjectElement();
+	}
+
+	if (!Object->ObservationObject.IsValid(Element.ObjectElement))
+	{
+		UE_LOG(LogLearning, Error, TEXT("MakeConv2dObservation: Invalid Observation Object."));
+		return FLearningAgentsObservationObjectElement();
+	}
+
+	return { Object->ObservationObject.CreateConv2d({ Element.ObjectElement }, Tag) };
+}
+
 FLearningAgentsObservationObjectElement ULearningAgentsObservations::MakeFloatObservation(
 	ULearningAgentsObservationObject* Object,
 	const float Value,
@@ -4585,3 +4678,38 @@ bool ULearningAgentsObservations::GetConv1dObservation(FLearningAgentsObservatio
 	return true;
 }
 
+bool ULearningAgentsObservations::GetConv2dObservation(FLearningAgentsObservationObjectElement& OutElement, const ULearningAgentsObservationObject* Object, const FLearningAgentsObservationObjectElement Element, const FName Tag)
+{
+	if (!Object)
+	{
+		UE_LOG(LogLearning, Error, TEXT("GetConv2dObservation: Object is nullptr."));
+		OutElement = FLearningAgentsObservationObjectElement();
+		return false;
+	}
+
+	if (!Object->ObservationObject.IsValid(Element.ObjectElement))
+	{
+		UE_LOG(LogLearning, Error, TEXT("GetConv2dObservation: Invalid Observation Object."));
+		OutElement = FLearningAgentsObservationObjectElement();
+		return false;
+	}
+
+	if (Object->ObservationObject.GetTag(Element.ObjectElement) != Tag)
+	{
+		UE_LOG(LogLearning, Warning, TEXT("GetConv2dObservation: Observation tag does not match. Observation is '%s' but asked for '%s'."), *Object->ObservationObject.GetTag(Element.ObjectElement).ToString(), *Tag.ToString());
+	}
+
+	if (Object->ObservationObject.GetType(Element.ObjectElement) != UE::Learning::Observation::EType::Conv2d)
+	{
+		UE_LOG(LogLearning, Error, TEXT("GetConv2dObservation: Observation '%s' type does not match. Observation is '%s' but asked for '%s'."),
+			*Object->ObservationObject.GetTag(Element.ObjectElement).ToString(),
+			UE::Learning::Agents::Observation::Private::GetObservationTypeString(Object->ObservationObject.GetType(Element.ObjectElement)),
+			UE::Learning::Agents::Observation::Private::GetObservationTypeString(UE::Learning::Observation::EType::Conv2d));
+		OutElement = FLearningAgentsObservationObjectElement();
+		return false;
+	}
+
+	OutElement = { Object->ObservationObject.GetConv2d(Element.ObjectElement).Element };
+
+	return true;
+}
diff --git a/Engine/Plugins/Experimental/LearningAgents/Source/LearningAgents/Public/LearningAgentsObservations.h b/Engine/Plugins/Experimental/LearningAgents/Source/LearningAgents/Public/LearningAgentsObservations.h
--- a/Engine/Plugins/Experimental/LearningAgents/Source/LearningAgents/Public/LearningAgentsObservations.h
+++ b/Engine/Plugins/Experimental/LearningAgents/Source/LearningAgents/Public/LearningAgentsObservations.h
@@ -533,19 +533,39 @@ public:
 	static UE_API FLearningAgentsObservationSchemaElement SpecifyBoolObservation(ULearningAgentsObservationSchema* Schema, const FName Tag = TEXT("BoolObservation"));
 
 	/**
-	 * Specifies a new 1d convolution observation of another observation.
+	 * Specifies a new 1d convolution observation of another observation (Flat buffer input, Channel x Length).
 	 *
 	 * @param Schema The Observation Schema
 	 * @param Element The sub-observation to be convolved.
-	 * @param InputLength Length of the convolution input.
-	 * @param InChannels Number of channels in the input.
-	 * @param OutChannels Number of channels produced by the convolution.
-	 * @param KernelSize Size of the convolving kernel.
+	 * @param InputLength The length of the convolution input.
+	 * @param InChannels The Number of channels in the input.
+	 * @param OutChannels The Number of channels produced by the convolution.
+	 * @param KernelSize The size of the convolving kernel.
+	 * @param Padding The Amount of padding applied to the input.
+	 * @param PaddingMode 'Zeros' or 'Circular'
 	 * @param Tag The tag of this new observation. Used during observation object validation and debugging.
 	 * @return The newly created observation schema element.
 	 */
 	UFUNCTION(BlueprintPure, Category = "LearningAgents", meta = (AdvancedDisplay = 9))
-	static UE_API FLearningAgentsObservationSchemaElement SpecifyConv1dObservation(ULearningAgentsObservationSchema* Schema, const FLearningAgentsObservationSchemaElement Element, const int32 InputLength, const int32 InChannels, const int32 OutChannels, const int32 KernelSize, const int32 Padding, const ELearningAgentsPaddingMode PaddingMode, const ELearningAgentsActivationFunction ActivationFunction, const FName Tag = TEXT("Conv1dObservation"));
+	static UE_API FLearningAgentsObservationSchemaElement SpecifyConv1dObservation(ULearningAgentsObservationSchema* Schema, const FLearningAgentsObservationSchemaElement Element, const int32 InputLength, const int32 InChannels, const int32 OutChannels, const int32 KernelSize, const int32 Padding, const ELearningAgentsPaddingMode PaddingMode = ELearningAgentsPaddingMode::Circular, const ELearningAgentsActivationFunction ActivationFunction = ELearningAgentsActivationFunction::ReLU, const FName Tag = TEXT("Conv1dObservation"));
+
+	/**
+	 * Specifies a new 2d convolution observation of another observation (Flat buffer input, Channels x Input Height x Input Width), Zeros Padding.
+	 *
+	 * @param Schema The Observation Schema
+	 * @param Element The sub-observation to be convolved.
+	 * @param InputHeight The height of the convolution input.
+	 * @param InputWidth The width of the convolution input.
+	 * @param InChannels The Number of channels in the input.
+	 * @param OutChannels The Number of channels produced by the convolution.
+	 * @param KernelSize The size of the convolving kernel (3 -> (3x3), 5 -> (5x5) etc).
+	 * @param Stride The stride of the convolution.
+	 * @param Padding The amount of padding applied to the input.
+	 * @param Tag The tag of this new observation. Used during observation object validation and debugging.
+	 * @return The newly created observation schema element.
+	 */
+	 UFUNCTION(BlueprintPure, Category = "LearningAgents", meta = (AdvancedDisplay = 10))
+	static UE_API FLearningAgentsObservationSchemaElement SpecifyConv2dObservation(ULearningAgentsObservationSchema* Schema, const FLearningAgentsObservationSchemaElement Element, const int32 InputHeight, const int32 InputWidth, const int32 InChannels, const int32 OutChannels, const int32 KernelSize, const int32 Stride = 1, const int32 Padding = 0, const ELearningAgentsActivationFunction ActivationFunction = ELearningAgentsActivationFunction::ReLU, const FName Tag = TEXT("Conv2dObservation"));
 
 	/**
 	 * Specifies a new float observation. A simple observation which can be used as a catch-all for situations where a
@@ -1280,6 +1300,15 @@ public:
 	UFUNCTION(BlueprintPure, Category = "LearningAgents", meta = (AdvancedDisplay = 2, DefaultToSelf = "VisualLoggerListener"))
 	static UE_API FLearningAgentsObservationObjectElement MakeConv1dObservation(ULearningAgentsObservationObject* Object, const FLearningAgentsObservationObjectElement Element, const FName Tag = TEXT("Conv1dObservation"));
 
+	/**
+	 * Make a new 2d convolution observation.
+	 *
+	 * @param Object The Observation Object
+	 * @param Element The Observation Element to be encoded.
+	 * @return The newly created observation object element.
+	 */
+	UFUNCTION(BlueprintPure, Category = "LearningAgents", meta = (AdvancedDisplay = 2, DefaultToSelf = "VisualLoggerListener"))
+	static UE_API FLearningAgentsObservationObjectElement MakeConv2dObservation(ULearningAgentsObservationObject* Object, const FLearningAgentsObservationObjectElement Element, const FName Tag = TEXT("Conv2dObservation"));
 
 	/**
 	 * Make a new float observation.
@@ -1863,6 +1892,9 @@ public:
 
 	UFUNCTION(BlueprintPure = false, Category = "LearningAgents", meta = (AdvancedDisplay = 3, ReturnDisplayName = "Success"))
 	static UE_API bool GetConv1dObservation(FLearningAgentsObservationObjectElement& OutElement, const ULearningAgentsObservationObject* Object, const FLearningAgentsObservationObjectElement Element, const FName Tag = TEXT("Conv1dObservation"));
+	
+	UFUNCTION(BlueprintPure = false, Category = "LearningAgents", meta = (AdvancedDisplay = 3, ReturnDisplayName = "Success"))
+	static UE_API bool GetConv2dObservation(FLearningAgentsObservationObjectElement& OutElement, const ULearningAgentsObservationObject* Object, const FLearningAgentsObservationObjectElement Element, const FName Tag = TEXT("Conv2dObservation"));
 };
 
 
diff --git a/Engine/Plugins/Experimental/LearningAgents/Source/LearningTraining/Private/LearningTrainer.cpp b/Engine/Plugins/Experimental/LearningAgents/Source/LearningTraining/Private/LearningTrainer.cpp
--- a/Engine/Plugins/Experimental/LearningAgents/Source/LearningTraining/Private/LearningTrainer.cpp
+++ b/Engine/Plugins/Experimental/LearningAgents/Source/LearningTraining/Private/LearningTrainer.cpp
@@ -281,6 +281,22 @@ namespace UE::Learning::Trainer
 			break;
 		}
 
+		case Observation::EType::Conv2d:
+		{
+			const Observation::FSchemaConv2dParameters Parameters = ObservationSchema.GetConv2d(ObservationSchemaElement);
+			Object->SetStringField(TEXT("Type"), TEXT("Conv2d"));
+			Object->SetNumberField(TEXT("InputHeight"), Parameters.InputHeight);
+			Object->SetNumberField(TEXT("InputWidth"), Parameters.InputWidth);
+			Object->SetNumberField(TEXT("InChannels"), Parameters.InChannels);
+			Object->SetNumberField(TEXT("OutChannels"), Parameters.OutChannels);
+			Object->SetNumberField(TEXT("KernelSize"), Parameters.KernelSize);
+			Object->SetNumberField(TEXT("Stride"), Parameters.Stride);
+			Object->SetNumberField(TEXT("Padding"), Parameters.Padding);
+			Object->SetNumberField(TEXT("PaddingMode"), (uint32)Parameters.Padding);
+			Object->SetObjectField(TEXT("Element"), ConvertObservationSchemaToJSON(ObservationSchema, Parameters.Element));
+			break;
+		}
+
 		default:
 			checkNoEntry();
 		}
diff --git a/Engine/Plugins/Experimental/NNERuntimeBasicCpu/Source/NNERuntimeBasicCpu/Private/NNERuntimeBasicCpu.ispc b/Engine/Plugins/Experimental/NNERuntimeBasicCpu/Source/NNERuntimeBasicCpu/Private/NNERuntimeBasicCpu.ispc
--- a/Engine/Plugins/Experimental/NNERuntimeBasicCpu/Source/NNERuntimeBasicCpu/Private/NNERuntimeBasicCpu.ispc
+++ b/Engine/Plugins/Experimental/NNERuntimeBasicCpu/Source/NNERuntimeBasicCpu/Private/NNERuntimeBasicCpu.ispc
@@ -755,4 +755,68 @@ export void NNERuntimeBasicCPUOperatorConv1d(
 			}
 		}
 	}
-}
\ No newline at end of file
+}
+
+export void NNERuntimeBasicCPUOperatorConv2d(
+	uniform float Output[],
+	uniform const float Input[],
+	uniform const float Weights[],
+	uniform const float Biases[],
+	uniform const int BatchNum,
+	uniform const int InChannelsNum,
+	uniform const int OutChannelsNum,
+	uniform const int InputHeight,
+	uniform const int InputWidth,
+	uniform const int KernelSize,
+	uniform const int Stride,
+	uniform const int Padding,
+	uniform const int PaddingMode,
+	uniform const int OutputStride,
+	uniform const int InputStride)
+{
+	uniform int OutputHeight = (InputHeight + 2 * Padding - KernelSize) / Stride + 1;
+	uniform int OutputWidth = (InputWidth + 2 * Padding - KernelSize) / Stride + 1;
+
+	for (uniform int BatchIdx = 0; BatchIdx < BatchNum; BatchIdx++)
+	{
+		uniform const float* InputBatchItr = Input + BatchIdx * InputStride;
+		uniform float* OutputBatchItr = Output + BatchIdx * OutputStride;
+
+		for (uniform int OutChannelsIdx = 0; OutChannelsIdx < OutChannelsNum; OutChannelsIdx++)
+		{
+			uniform const float* WeightsItr = Weights + OutChannelsIdx * InChannelsNum * KernelSize * KernelSize;
+			uniform float Bias = Biases[OutChannelsIdx];
+
+			for (uniform int OutputHeightIdx = 0; OutputHeightIdx < OutputHeight; OutputHeightIdx++)
+			{
+				for (uniform int OutputWidthIdx = 0; OutputWidthIdx < OutputWidth; OutputWidthIdx++)
+				{
+					float Sum = Bias;
+
+					for (uniform int InChannelsIdx = 0; InChannelsIdx < InChannelsNum; InChannelsIdx++)
+					{
+						uniform const float* WeightsInChannelsItr = WeightsItr + InChannelsIdx * KernelSize * KernelSize;
+						uniform const float* InChannelsItr = InputBatchItr + InChannelsIdx * InputHeight * InputWidth;
+
+						for (uniform int KernelHeight = 0; KernelHeight < KernelSize; KernelHeight++)
+						{
+							for (uniform int KernelWidth = 0; KernelWidth < KernelSize; KernelWidth++)
+							{
+								uniform int HeightIdx = OutputHeightIdx * Stride + KernelHeight - Padding;
+								uniform int WidthIdx = OutputWidthIdx * Stride + KernelWidth - Padding;
+								
+								float Val = 0.0f;
+								if (HeightIdx >= 0 && HeightIdx < InputHeight && WidthIdx >= 0 && WidthIdx < InputWidth)
+								{
+									Val = InChannelsItr[HeightIdx * InputWidth + WidthIdx];
+								}
+								Sum += Val * WeightsInChannelsItr[KernelHeight * KernelSize + KernelWidth];
+							}
+						}
+					}
+					OutputBatchItr[OutChannelsIdx * OutputHeight * OutputWidth + OutputHeightIdx * OutputWidth + OutputWidthIdx] = Sum;
+				}
+			}
+		}
+	}
+}
diff --git a/Engine/Plugins/Experimental/NNERuntimeBasicCpu/Source/NNERuntimeBasicCpu/Private/NNERuntimeBasicCpuModel.cpp b/Engine/Plugins/Experimental/NNERuntimeBasicCpu/Source/NNERuntimeBasicCpu/Private/NNERuntimeBasicCpuModel.cpp
--- a/Engine/Plugins/Experimental/NNERuntimeBasicCpu/Source/NNERuntimeBasicCpu/Private/NNERuntimeBasicCpuModel.cpp
+++ b/Engine/Plugins/Experimental/NNERuntimeBasicCpu/Source/NNERuntimeBasicCpu/Private/NNERuntimeBasicCpuModel.cpp
@@ -521,6 +521,91 @@ namespace UE::NNE::RuntimeBasic
 #endif
 		}
 
+		static inline void OperatorConv2d(
+			float* RESTRICT Output,
+			const float* RESTRICT Input,
+			const float* RESTRICT Weights,
+			const float* RESTRICT Biases,
+			const uint32 BatchSize,
+			const uint32 InChannels,
+			const uint32 OutChannels,
+			const uint32 InputHeight,
+			const uint32 InputWidth,
+			const uint32 KernelSize,
+			const uint32 Stride,
+			const uint32 Padding,
+			const FModelBuilder::EPaddingMode PaddingMode,
+			const uint32 OutputStride,
+			const uint32 InputStride)
+		{
+			NNE_RUNTIME_BASIC_TRACE_SCOPE(NNE::RuntimeBasic::Private::OperatorConv2d);
+			NNE_RUNTIME_BASIC_CHECK_ALIASING_SIZES(Output, Input, ((InputHeight + 2 * Padding - KernelSize) / Stride + 1) * ((InputWidth + 2 * Padding - KernelSize) / Stride + 1) * OutChannels, InputHeight * InputWidth * InChannels);
+
+			check(PaddingMode == FModelBuilder::EPaddingMode::Zeros);
+			check((InputHeight + 2 * Padding - KernelSize) % Stride == 0);
+
+#if NNE_RUNTIME_BASIC_ENABLE_ISPC
+			ispc::NNERuntimeBasicCPUOperatorConv2d(
+				Output,
+				Input,
+				Weights,
+				Biases,
+				BatchSize,
+				InChannels,
+				OutChannels,
+				InputHeight,
+				InputWidth,
+				KernelSize,
+				Stride,
+				Padding,
+				static_cast<uint32>(PaddingMode),
+				OutputStride,
+				InputStride);
+#else
+			const uint32 OutputHeight = (InputHeight + 2 * Padding - KernelSize) / Stride + 1;
+			const uint32 OutputWidth = (InputWidth + 2 * Padding - KernelSize) / Stride + 1;
+
+			for (uint32 BatchIdx = 0; BatchIdx < BatchSize; BatchIdx++)
+			{
+				const float* InputBatchItr = Input + BatchIdx * InputStride;
+				float* OutputBatchItr = Output + BatchIdx * OutputStride;
+
+				for (uint32 OutChannelsIdx = 0; OutChannelsIdx < OutChannels; OutChannelsIdx++)
+				{
+					const float Bias = Biases[OutChannelsIdx];
+					const float* WeightsItr = Weights + OutChannelsIdx * InChannels * KernelSize * KernelSize;
+
+					for (uint32 OutputHeightIdx = 0; OutputHeightIdx < OutputHeight; OutputHeightIdx++)
+					{
+						for (uint32 OutputWidthIdx = 0; OutputWidthIdx < OutputWidth; OutputWidthIdx++)
+						{	
+							float Sum = Bias;
+
+							for (uint32 InChannelsIdx = 0; InChannelsIdx < InChannels; InChannelsIdx++)
+							{
+								const float* WeightsInChannelsItr = WeightsItr + InChannelsIdx * KernelSize * KernelSize;
+								const float* InChannelsItr = InputBatchItr + InChannelsIdx * InputHeight * InputWidth;
+
+								for (uint32 KernelHeightIdx = 0; KernelHeightIdx < KernelSize; KernelHeightIdx++)
+								{
+									for (uint32 KernelWidthIdx = 0; KernelWidthIdx < KernelSize; KernelWidthIdx++)
+									{
+										int32 HeightIdx = OutputHeightIdx * Stride + KernelHeightIdx - Padding;
+										int32 WidthIdx = OutputWidthIdx * Stride + KernelWidthIdx - Padding;
+
+										float Val = InChannelsItr[HeightIdx * InputWidth + WidthIdx];
+										Sum += Val * WeightsInChannelsItr[KernelHeightIdx * KernelSize + KernelWidthIdx];
+									}
+								}
+							}
+							OutputBatchItr[OutChannelsIdx * OutputHeight * OutputWidth + OutputHeightIdx * OutputWidth + OutputWidthIdx] = Sum;
+						}
+					}
+				}
+			}
+#endif
+		}
+
 		static inline void OperatorLinear(
 			float* RESTRICT Output,
 			const float* RESTRICT Input,
@@ -1764,6 +1849,7 @@ namespace UE::NNE::RuntimeBasic
 			FiLM = 27,
 			GRUCell = 28,
 			Conv1d = 29,
+			Conv2d = 30,
 		};
 
 		//--------------------------------------------------------------------------
@@ -5204,7 +5290,7 @@ namespace UE::NNE::RuntimeBasic
 
 				OperatorNanCheck(OutputBuffer, BatchSize, OutputBufferSize, OutputBufferStride);
 			}
-			uint32 InputLength, InChannels, OutChannels, KernelSize, Padding;
+			uint32 InputLength = 0, InChannels = 0, OutChannels = 0, KernelSize = 0, Padding = 0;
 			FModelBuilder::EPaddingMode PaddingMode;
 			TConstArrayView<float> Weights, Biases;
 			TSharedPtr<ILayer> SubLayer;
@@ -5221,7 +5307,153 @@ namespace UE::NNE::RuntimeBasic
 			SubLayerInstance->SetMaxBatchSize(MaxBatchSize);
 			SubLayerOutputBuffer.SetNumUninitialized(MaxBatchSize * Conv1dLayer.SubLayer->GetOutputSize(), EAllowShrinking::No);
 		}
-		
+		//--------------------------------------------------------------------------
+
+		struct FConv2dLayer;
+
+		struct FConv2dInstance : public ILayerInstance
+		{
+			FConv2dInstance(const FConv2dLayer& InConv2dLayer);
+
+			void SetMaxBatchSize(const uint32 MaxBatchSize) override final;
+			virtual ELayerType GetLayerType() const override final { return ELayerType::Conv2d; }
+
+			const FConv2dLayer& Conv2dLayer;
+			TSharedPtr<ILayerInstance> SubLayerInstance;
+			TArray<float> SubLayerOutputBuffer;
+		};
+
+		struct FConv2dLayer : public ILayer
+		{
+			virtual TSharedPtr<ILayerInstance> MakeInstance() const { return MakeShared<FConv2dInstance>(*this); };
+			virtual ELayerType GetLayerType() const override final { return ELayerType::Conv2d; }
+			virtual uint32 GetInputSize() const override final { return SubLayer->GetInputSize(); }
+			virtual uint32 GetOutputSize() const override final
+			{
+				uint32 OutputHeight = (InputHeight + 2 * Padding - KernelSize) / Stride + 1;
+				uint32 OutputWidth = (InputWidth + 2 * Padding - KernelSize) / Stride + 1;
+				return OutputHeight * OutputWidth * OutChannels;
+			}
+			uint32 GetPaddingMode() const { return static_cast<uint32_t>(PaddingMode); }
+
+			void SerializationLoadPaddingMode(uint64& InOutOffset, TConstArrayView<uint8> Data)
+			{
+				uint32 LoadedMode;
+				Serialization::Load(InOutOffset, LoadedMode, Data);
+				PaddingMode = static_cast<FModelBuilder::EPaddingMode>(LoadedMode);
+			}
+
+			virtual void SerializationSize(uint64& InOutOffset) const override final
+			{
+				Serialization::Size(InOutOffset, InputHeight);
+				Serialization::Size(InOutOffset, InputWidth);
+				Serialization::Size(InOutOffset, InChannels);
+				Serialization::Size(InOutOffset, OutChannels);
+				Serialization::Size(InOutOffset, KernelSize);
+				Serialization::Size(InOutOffset, Stride);
+				Serialization::Size(InOutOffset, Padding);
+				Serialization::Size(InOutOffset, GetPaddingMode());
+				Serialization::Size(InOutOffset, Weights);
+				Serialization::Size(InOutOffset, Biases);
+				Serialization::Size(InOutOffset, SubLayer);
+			}
+
+			virtual void SerializationLoad(uint64& InOutOffset, TConstArrayView<uint8> Data) override final
+			{
+				Serialization::Load(InOutOffset, InputHeight, Data);
+				Serialization::Load(InOutOffset, InputWidth, Data);
+				Serialization::Load(InOutOffset, InChannels, Data);
+				Serialization::Load(InOutOffset, OutChannels, Data);
+				Serialization::Load(InOutOffset, KernelSize, Data);
+				Serialization::Load(InOutOffset, Stride, Data);
+				Serialization::Load(InOutOffset, Padding, Data);
+				SerializationLoadPaddingMode(InOutOffset, Data);
+				Serialization::Load(InOutOffset, Weights, Data, KernelSize * KernelSize * InChannels * OutChannels);
+				Serialization::Load(InOutOffset, Biases, Data, OutChannels);
+				Serialization::Load(InOutOffset, SubLayer, Data);
+			}
+
+			virtual void SerializationSave(uint64& InOutOffset, TArrayView<uint8> Data) const override final
+			{
+				Serialization::Save(InOutOffset, InputHeight, Data);
+				Serialization::Save(InOutOffset, InputWidth, Data);
+				Serialization::Save(InOutOffset, InChannels, Data);
+				Serialization::Save(InOutOffset, OutChannels, Data);
+				Serialization::Save(InOutOffset, KernelSize, Data);
+				Serialization::Save(InOutOffset, Stride, Data);
+				Serialization::Save(InOutOffset, Padding, Data);
+				Serialization::Save(InOutOffset, GetPaddingMode(), Data);
+				Serialization::Save(InOutOffset, Weights, Data);
+				Serialization::Save(InOutOffset, Biases, Data);
+				Serialization::Save(InOutOffset, SubLayer, Data);
+			}
+
+			virtual void Evaluate(
+				ILayerInstance* Instance,
+				float* OutputBuffer,
+				const float* InputBuffer,
+				const uint32 BatchSize,
+				const uint32 OutputBufferSize,
+				const uint32 InputBufferSize,
+				const uint32 OutputBufferStride,
+				const uint32 InputBufferStride) override final
+			{
+				NNE_RUNTIME_BASIC_TRACE_SCOPE(NNE::RuntimeBasic::Private::FConv2dLayer::Evaluate);
+				check(OutputBufferSize == GetOutputSize() && InputBufferSize == GetInputSize());
+				check(OutputBufferStride >= GetOutputSize() && InputBufferStride >= GetInputSize());
+				check(Instance && Instance->GetLayerType() == GetLayerType());
+
+				FConv2dInstance* Conv2dLayerInstance = StaticCast<FConv2dInstance*>(Instance);
+
+				OperatorNanCheck(InputBuffer, BatchSize, InputBufferSize, InputBufferStride);
+
+				SubLayer->Evaluate(
+					Conv2dLayerInstance->SubLayerInstance.Get(),
+					Conv2dLayerInstance->SubLayerOutputBuffer.GetData(),
+					InputBuffer,
+					BatchSize,
+					SubLayer->GetOutputSize(),
+					SubLayer->GetInputSize(),
+					SubLayer->GetOutputSize(),
+					SubLayer->GetInputSize());
+
+				OperatorConv2d(
+					OutputBuffer,
+					Conv2dLayerInstance->SubLayerOutputBuffer.GetData(),
+					Weights.GetData(),
+					Biases.GetData(),
+					BatchSize,
+					InChannels,
+					OutChannels,
+					InputHeight,
+					InputWidth,
+					KernelSize,
+					Stride,
+					Padding,
+					PaddingMode,
+					OutputBufferStride,
+					InputHeight * InputWidth * InChannels);
+
+				OperatorNanCheck(OutputBuffer, BatchSize, OutputBufferSize, OutputBufferStride);
+			}
+			uint32 InputHeight = 0, InputWidth = 0, InChannels = 0, OutChannels = 0, KernelSize = 0, Stride = 0, Padding = 0;
+			FModelBuilder::EPaddingMode PaddingMode;
+			TConstArrayView<float> Weights, Biases;
+			TSharedPtr<ILayer> SubLayer;
+		};
+
+		FConv2dInstance::FConv2dInstance(const FConv2dLayer& InConv2dLayer) : Conv2dLayer(InConv2dLayer)
+		{
+			SubLayerInstance = Conv2dLayer.SubLayer->MakeInstance();
+		}
+
+		void FConv2dInstance::SetMaxBatchSize(const uint32 MaxBatchSize)
+		{
+			NNE_RUNTIME_BASIC_TRACE_SCOPE(NNE::RuntimeBasic::Private::FConv2dLayerInstance::SetMaxBatchSize);
+			SubLayerInstance->SetMaxBatchSize(MaxBatchSize);
+			SubLayerOutputBuffer.SetNumUninitialized(MaxBatchSize * Conv2dLayer.SubLayer->GetOutputSize(), EAllowShrinking::No);
+		}
+
 		//--------------------------------------------------------------------------
 		// Layer Serialization
 		//--------------------------------------------------------------------------
@@ -5279,6 +5511,7 @@ namespace UE::NNE::RuntimeBasic
 					case ELayerType::FiLM: OutLayer = MakeShared<FFiLMLayer>(); break;
 					case ELayerType::GRUCell: OutLayer = MakeShared<FGRUCellLayer>(); break;
 					case ELayerType::Conv1d: OutLayer = MakeShared<FConv1dLayer>(); break;
+					case ELayerType::Conv2d: OutLayer = MakeShared<FConv2dLayer>(); break;
 					default: checkf(false, TEXT("Unknown Layer Id %i"), LayerTypeId);
 					}
 				}
@@ -6474,7 +6707,6 @@ namespace UE::NNE::RuntimeBasic
 		const uint32 KernelSize,
 		const uint32 Padding,
 		const EPaddingMode PaddingMode,
-		const EActivationFunction ActivationFunction,
 		const TConstArrayView<float> Weights,
 		const TConstArrayView<float> Biases,
 		const FModelBuilderElement& SubLayer)
@@ -6492,12 +6724,40 @@ namespace UE::NNE::RuntimeBasic
 		Conv1dLayer->Weights = Weights;
 		Conv1dLayer->Biases = Biases;
 		Conv1dLayer->SubLayer = SubLayer.Layer;
-
-		TArray<FModelBuilderElement, TInlineAllocator<32>> Layers;
-		Layers.Emplace(StaticCastSharedPtr<Private::ILayer>(Conv1dLayer));
-		Layers.Emplace(MakeActivation(InputLength * OutChannels, ActivationFunction));
 		
-		return MakeSequence(Layers);
+		return StaticCastSharedPtr<Private::ILayer>(Conv1dLayer);
+	}
+
+	FModelBuilderElement FModelBuilder::MakeConv2d(
+		const uint32 InputHeight,
+		const uint32 InputWidth,
+		const uint32 InChannels,
+		const uint32 OutChannels,
+		const uint32 KernelSize,
+		const uint32 Stride,
+		const uint32 Padding,
+		const EPaddingMode PaddingMode,
+		const TConstArrayView<float> Weights,
+		const TConstArrayView<float> Biases,
+		const FModelBuilderElement& SubLayer)
+	{
+		check(Weights.Num() == InChannels * OutChannels * KernelSize * KernelSize);
+		check(Biases.Num() == OutChannels);
+
+		const TSharedPtr<Private::FConv2dLayer> Conv2dLayer = MakeShared<Private::FConv2dLayer>();
+		Conv2dLayer->InputHeight = InputHeight;
+		Conv2dLayer->InputWidth = InputWidth;
+		Conv2dLayer->InChannels = InChannels;
+		Conv2dLayer->OutChannels = OutChannels;
+		Conv2dLayer->KernelSize = KernelSize;
+		Conv2dLayer->Stride = Stride;
+		Conv2dLayer->Padding = Padding;
+		Conv2dLayer->PaddingMode = PaddingMode;
+		Conv2dLayer->Weights = Weights;
+		Conv2dLayer->Biases = Biases;
+		Conv2dLayer->SubLayer = SubLayer.Layer;
+
+		return StaticCastSharedPtr<Private::ILayer>(Conv2dLayer);
 	}
 
 	void FModelBuilder::Reset()
diff --git a/Engine/Plugins/Experimental/NNERuntimeBasicCpu/Source/NNERuntimeBasicCpu/Public/NNERuntimeBasicCpuBuilder.h b/Engine/Plugins/Experimental/NNERuntimeBasicCpu/Source/NNERuntimeBasicCpu/Public/NNERuntimeBasicCpuBuilder.h
--- a/Engine/Plugins/Experimental/NNERuntimeBasicCpu/Source/NNERuntimeBasicCpu/Public/NNERuntimeBasicCpuBuilder.h
+++ b/Engine/Plugins/Experimental/NNERuntimeBasicCpu/Source/NNERuntimeBasicCpu/Public/NNERuntimeBasicCpuBuilder.h
@@ -569,6 +569,7 @@ namespace UE::NNE::RuntimeBasic
 		 * @param PaddingMode			"Zeros" or "Circular"
 		 * @param Weights				Convolutional layer weights
 		 * @param Biases				Convolutional layer biases
+		 * @param SubLayer				The sublayer to evaluate on the input
 		 */
 		UE_API FModelBuilderElement MakeConv1d(
 			const uint32 InputLength,
@@ -577,7 +578,34 @@ namespace UE::NNE::RuntimeBasic
 			const uint32 KernelSize,
 			const uint32 Padding, 
 			const EPaddingMode PaddingMode,
-			const EActivationFunction ActivationFunction, 
+			const TConstArrayView<float> Weights,
+			const TConstArrayView<float> Biases,
+			const FModelBuilderElement& SubLayer);
+
+		/**
+		* Make a new 2D Convolution Layer.
+		*
+		* @param InputHeight			Height of the 2d flat buffer input
+		* @param InputWidth				Width of the 2d flat buffer input
+		* @param InChannels				Number of channels in the input
+		* @param OutChannels			Number of channels produced by the convolution
+		* @param KernelSize				Size of the convolving kernel
+		* @param Stride					The stride of the cross correlation 
+		* @param Padding				The amount of padding added to both sides of the input
+		* @param PaddingMode			"Zeros" or "Circular"
+		* @param Weights				Convolutional layer weights
+		* @param Biases					Convolutional layer biases
+		* @param SubLayer				The sublayer to evaluate on the input
+		*/
+		UE_API FModelBuilderElement MakeConv2d(
+			const uint32 InputHeight,
+			const uint32 InputWidth,
+			const uint32 InChannels,
+			const uint32 OutChannels,
+			const uint32 KernelSize,
+			const uint32 Stride,
+			const uint32 Padding,
+			const EPaddingMode PaddingMode,
 			const TConstArrayView<float> Weights,
 			const TConstArrayView<float> Biases,
 			const FModelBuilderElement& SubLayer);
