[NNERuntimeBasicCpu] ISPC optimization

#

diff --git a/Engine/Plugins/Experimental/NNERuntimeBasicCpu/Source/NNERuntimeBasicCpu/Private/NNERuntimeBasicCpu.ispc b/Engine/Plugins/Experimental/NNERuntimeBasicCpu/Source/NNERuntimeBasicCpu/Private/NNERuntimeBasicCpu.ispc
--- a/Engine/Plugins/Experimental/NNERuntimeBasicCpu/Source/NNERuntimeBasicCpu/Private/NNERuntimeBasicCpu.ispc
+++ b/Engine/Plugins/Experimental/NNERuntimeBasicCpu/Source/NNERuntimeBasicCpu/Private/NNERuntimeBasicCpu.ispc
@@ -757,6 +757,7 @@ export void NNERuntimeBasicCPUOperatorConv1d(
 	}
 }
 
+
 export void NNERuntimeBasicCPUOperatorConv2d(
 	uniform float Output[],
 	uniform const float Input[],
@@ -787,36 +788,34 @@ export void NNERuntimeBasicCPUOperatorConv2d(
 			uniform const float* WeightsItr = Weights + OutChannelsIdx * InChannelsNum * KernelSize * KernelSize;
 			uniform float Bias = Biases[OutChannelsIdx];
 
-			for (uniform int OutputHeightIdx = 0; OutputHeightIdx < OutputHeight; OutputHeightIdx++)
+			foreach (OutputHeightIdx = 0 ... OutputHeight, OutputWidthIdx = 0 ... OutputWidth)
 			{
-				for (uniform int OutputWidthIdx = 0; OutputWidthIdx < OutputWidth; OutputWidthIdx++)
-				{
-					float Sum = Bias;
-
-					for (uniform int InChannelsIdx = 0; InChannelsIdx < InChannelsNum; InChannelsIdx++)
-					{
-						uniform const float* WeightsInChannelsItr = WeightsItr + InChannelsIdx * KernelSize * KernelSize;
-						uniform const float* InChannelsItr = InputBatchItr + InChannelsIdx * InputHeight * InputWidth;
-
-						for (uniform int KernelHeight = 0; KernelHeight < KernelSize; KernelHeight++)
-						{
-							for (uniform int KernelWidth = 0; KernelWidth < KernelSize; KernelWidth++)
-							{
-								uniform int HeightIdx = OutputHeightIdx * Stride + KernelHeight - Padding;
-								uniform int WidthIdx = OutputWidthIdx * Stride + KernelWidth - Padding;
-								
-								float Val = 0.0f;
-								if (HeightIdx >= 0 && HeightIdx < InputHeight && WidthIdx >= 0 && WidthIdx < InputWidth)
-								{
-									Val = InChannelsItr[HeightIdx * InputWidth + WidthIdx];
-								}
-								Sum += Val * WeightsInChannelsItr[KernelHeight * KernelSize + KernelWidth];
-							}
-						}
-					}
-					OutputBatchItr[OutChannelsIdx * OutputHeight * OutputWidth + OutputHeightIdx * OutputWidth + OutputWidthIdx] = Sum;
-				}
+                float Sum = Bias;
+
+                for (uniform int InChannelsIdx = 0; InChannelsIdx < InChannelsNum; InChannelsIdx++)
+                {
+                    uniform const float* WeightsInChannelsItr = WeightsItr + InChannelsIdx * KernelSize * KernelSize;
+                    uniform const float* InChannelsItr = InputBatchItr + InChannelsIdx * InputHeight * InputWidth;
+
+                    for (uniform int KernelHeight = 0; KernelHeight < KernelSize; KernelHeight++)
+                    {
+                        for (uniform int KernelWidth = 0; KernelWidth < KernelSize; KernelWidth++)
+                        {
+                            int HeightIdx = OutputHeightIdx * Stride + KernelHeight - Padding;
+                            int WidthIdx = OutputWidthIdx * Stride + KernelWidth - Padding;
+                            
+                            float Val = 0.0f;
+                            if (HeightIdx >= 0 && HeightIdx < InputHeight && WidthIdx >= 0 && WidthIdx < InputWidth)
+                            {
+								#pragma ignore warning
+                                Val = InChannelsItr[HeightIdx * InputWidth + WidthIdx];
+                            }
+                            Sum += Val * WeightsInChannelsItr[KernelHeight * KernelSize + KernelWidth];
+                        }
+                    }
+                }
+                OutputBatchItr[OutChannelsIdx * OutputHeight * OutputWidth + OutputHeightIdx * OutputWidth + OutputWidthIdx] = Sum;
 			}
 		}
 	}
-}
+}
\ No newline at end of file
