NNERuntimeBasicCpu: Moved LayerNorm layers to after the Linear Layer for LayerNorm MLPs constructor

diff --git a/Engine/Plugins/Experimental/NNERuntimeBasicCpu/Source/NNERuntimeBasicCpu/Private/NNERuntimeBasicCpuModel.cpp b/Engine/Plugins/Experimental/NNERuntimeBasicCpu/Source/NNERuntimeBasicCpu/Private/NNERuntimeBasicCpuModel.cpp
--- a/Engine/Plugins/Experimental/NNERuntimeBasicCpu/Source/NNERuntimeBasicCpu/Private/NNERuntimeBasicCpuModel.cpp
+++ b/Engine/Plugins/Experimental/NNERuntimeBasicCpu/Source/NNERuntimeBasicCpu/Private/NNERuntimeBasicCpuModel.cpp
@@ -6182,9 +6182,7 @@ namespace UE::NNE::RuntimeBasic
 			{
 				Layers.Emplace(MakeSpread({
 					MakeCopy(InputSize),
-					MakeSequence({
-						MakeLinearLayer(InputSize, HiddenSize, LinearLayerSettings)
-					})
+					MakeLinearLayer(InputSize, HiddenSize, LinearLayerSettings)
 					}));
 			}
 			else if (LayerIdx != LayerNum - 1)
@@ -6192,8 +6190,8 @@ namespace UE::NNE::RuntimeBasic
 				Layers.Emplace(MakeSpread({
 					MakeSlice(InputSize + HiddenSize, 0, InputSize),
 					MakeSequence({
-						MakeLayerNorm(InputSize + HiddenSize, MakeValuesZero(InputSize + HiddenSize), MakeValuesOne(InputSize + HiddenSize)),
 						MakeLinearLayer(InputSize + HiddenSize, HiddenSize, LinearLayerSettings),
+						MakeLayerNorm(HiddenSize, MakeValuesZero(HiddenSize), MakeValuesOne(HiddenSize)),
 						MakeActivation(HiddenSize, ActivationFunction),
 					})
 					}));
@@ -6203,17 +6201,12 @@ namespace UE::NNE::RuntimeBasic
 				if (bActivationOnFinalLayer)
 				{
 					Layers.Emplace(MakeSequence({
-							MakeLayerNorm(InputSize + HiddenSize, MakeValuesZero(InputSize + HiddenSize), MakeValuesOne(InputSize + HiddenSize)),
-							MakeLinearLayer(InputSize + HiddenSize, OutputSize, LinearLayerSettings),
-							MakeActivation(OutputSize, ActivationFunction),
-						}));
+						MakeLinearLayer(InputSize + HiddenSize, OutputSize, LinearLayerSettings),
+						MakeActivation(OutputSize, ActivationFunction) }));
 				}
 				else
 				{
-					Layers.Emplace(MakeSequence({
-							MakeLayerNorm(InputSize + HiddenSize, MakeValuesZero(InputSize + HiddenSize), MakeValuesOne(InputSize + HiddenSize)),
-							MakeLinearLayer(InputSize + HiddenSize, OutputSize, LinearLayerSettings)
-						}));
+					Layers.Emplace(MakeLinearLayer(InputSize + HiddenSize, OutputSize, LinearLayerSettings));
 				}
 			}
 		}
@@ -6308,8 +6301,8 @@ namespace UE::NNE::RuntimeBasic
 			{
 				Layers.Emplace(MakeResidual(
 					MakeSequence({
-						MakeLayerNorm(HiddenSize, MakeValuesZero(HiddenSize), MakeValuesOne(HiddenSize)),
 						MakeLinearLayer(HiddenSize, HiddenSize, LinearLayerSettings),
+						MakeLayerNorm(HiddenSize, MakeValuesZero(HiddenSize), MakeValuesOne(HiddenSize)),
 						MakeActivation(HiddenSize, ActivationFunction),
 						})));
 			}
@@ -6318,17 +6311,12 @@ namespace UE::NNE::RuntimeBasic
 				if (bActivationOnFinalLayer)
 				{
 					Layers.Emplace(MakeSequence({
-							MakeLayerNorm(HiddenSize, MakeValuesZero(HiddenSize), MakeValuesOne(HiddenSize)),
-							MakeLinearLayer(HiddenSize, OutputSize, LinearLayerSettings),
-							MakeActivation(OutputSize, ActivationFunction),
-						}));
+						MakeLinearLayer(HiddenSize, OutputSize, LinearLayerSettings),
+						MakeActivation(OutputSize, ActivationFunction) }));
 				}
 				else
 				{
-					Layers.Emplace(MakeSequence({
-						MakeLayerNorm(HiddenSize, MakeValuesZero(HiddenSize), MakeValuesOne(HiddenSize)),
-						MakeLinearLayer(HiddenSize, OutputSize, LinearLayerSettings)
-					}));
+					Layers.Emplace(MakeLinearLayer(HiddenSize, OutputSize, LinearLayerSettings));
 				}
 			}
 		}
