NNERuntimeBasicCpu: Added Softmax

#Synced-CL 49628397

diff --git a/Engine/Plugins/Experimental/NNERuntimeBasicCpu/Source/NNERuntimeBasicCpu/Private/NNERuntimeBasicCpu.ispc b/Engine/Plugins/Experimental/NNERuntimeBasicCpu/Source/NNERuntimeBasicCpu/Private/NNERuntimeBasicCpu.ispc
--- a/Engine/Plugins/Experimental/NNERuntimeBasicCpu/Source/NNERuntimeBasicCpu/Private/NNERuntimeBasicCpu.ispc
+++ b/Engine/Plugins/Experimental/NNERuntimeBasicCpu/Source/NNERuntimeBasicCpu/Private/NNERuntimeBasicCpu.ispc
@@ -294,6 +294,37 @@ export void NNERuntimeBasicCPUOperatorPReLU(
 	}
 }
 
+export void NNERuntimeBasicCPUOperatorSoftmax(
+	uniform float Output[],
+	uniform const float Input[],
+	uniform const int BatchNum,
+	uniform const int Num,
+	uniform const int OutputStride,
+	uniform const int InputStride)
+{
+	for (uniform int32 BatchIdx = 0; BatchIdx < BatchNum; BatchIdx++)
+	{
+		float Maxs = 0.0f;
+		foreach (Idx = 0 ... Num)
+		{
+			Maxs = max(Maxs, Input[BatchIdx * InputStride + Idx]);
+		}
+		uniform float Max = reduce_max(Maxs);
+
+		float Denoms = 0.0f;
+		foreach (Idx = 0 ... Num)
+		{
+			Denoms += exp(Input[BatchIdx * InputStride + Idx] - Max);
+		}
+		uniform float Denom = max(reduce_add(Denoms), 1e-8f);
+
+		foreach (Idx = 0 ... Num)
+		{
+			Output[BatchIdx * OutputStride + Idx] = exp(Input[BatchIdx * InputStride + Idx] - Max) / Denom;
+		}
+	}
+}
+
 export void NNERuntimeBasicCPUOperatorMemoryCellUpdateMemory(
 	uniform float Output[],
 	uniform const float RememberGate[],
diff --git a/Engine/Plugins/Experimental/NNERuntimeBasicCpu/Source/NNERuntimeBasicCpu/Private/NNERuntimeBasicCpuModel.cpp b/Engine/Plugins/Experimental/NNERuntimeBasicCpu/Source/NNERuntimeBasicCpu/Private/NNERuntimeBasicCpuModel.cpp
--- a/Engine/Plugins/Experimental/NNERuntimeBasicCpu/Source/NNERuntimeBasicCpu/Private/NNERuntimeBasicCpuModel.cpp
+++ b/Engine/Plugins/Experimental/NNERuntimeBasicCpu/Source/NNERuntimeBasicCpu/Private/NNERuntimeBasicCpuModel.cpp
@@ -932,6 +932,50 @@ namespace UE::NNE::RuntimeBasic
 #endif
 		}
 
+		static inline void OperatorSoftmax(
+			float* RESTRICT Output,
+			const float* RESTRICT Input,
+			const uint32 BatchSize,
+			const uint32 InputOutputSize,
+			const uint32 OutputStride,
+			const uint32 InputStride)
+		{
+			NNE_RUNTIME_BASIC_TRACE_SCOPE(NNE::RuntimeBasic::Private::OperatorSoftmax);
+
+			NNE_RUNTIME_BASIC_CHECK_ALIASING_SIZE(Output, Input, InputOutputSize);
+
+#if NNE_RUNTIME_BASIC_ENABLE_ISPC
+			ispc::NNERuntimeBasicCPUOperatorSoftmax(
+				Output,
+				Input,
+				BatchSize,
+				InputOutputSize,
+				OutputStride,
+				InputStride);
+#else
+			for (uint32 BatchIdx = 0; BatchIdx < BatchSize; BatchIdx++)
+			{
+				float Max = 0.0f;
+				for (uint32 Idx = 0; Idx < InputOutputSize; Idx++)
+				{
+					Max = FMath::Max(Max, Input[BatchIdx * InputStride + Idx]);
+				}
+
+				float Denom = 0.0f;
+				for (uint32 Idx = 0; Idx < InputOutputSize; Idx++)
+				{
+					Denom += FMath::Exp(Input[BatchIdx * InputStride + Idx] - Max);
+				}
+				Denom = FMath::Max(Denom, 1e-8f);
+
+				for (uint32 Idx = 0; Idx < InputOutputSize; Idx++)
+				{
+					Output[BatchIdx * OutputStride + Idx] = FMath::Exp(Input[BatchIdx * InputStride + Idx] - Max) / Denom;
+				}
+			}
+#endif
+		}
+
 		static inline void OperatorMemoryCellUpdateMemory(
 			float* RESTRICT Output,
 			const float* RESTRICT RememberGate,
@@ -1849,6 +1893,7 @@ namespace UE::NNE::RuntimeBasic
 			GRUCell = 28,
 			Conv1d = 29,
 			Conv2d = 30,
+			Softmax = 31,
 		};
 
 		//--------------------------------------------------------------------------
@@ -5453,6 +5498,60 @@ namespace UE::NNE::RuntimeBasic
 			SubLayerOutputBuffer.SetNumUninitialized(MaxBatchSize * Conv2dLayer.SubLayer->GetOutputSize(), EAllowShrinking::No);
 		}
 
+		//--------------------------------------------------------------------------
+
+		struct FSoftmaxLayer : public ILayer
+		{
+			virtual ELayerType GetLayerType() const override final { return ELayerType::Softmax; }
+			virtual uint32 GetInputSize() const override final { return InputOutputSize; }
+			virtual uint32 GetOutputSize() const override final { return InputOutputSize; }
+
+			virtual void SerializationSize(uint64& InOutOffset) const override final
+			{
+				Serialization::Size(InOutOffset, InputOutputSize);
+			}
+
+			virtual void SerializationLoad(uint64& InOutOffset, TConstArrayView<uint8> Data) override final
+			{
+				Serialization::Load(InOutOffset, InputOutputSize, Data);
+			}
+
+			virtual void SerializationSave(uint64& InOutOffset, TArrayView<uint8> Data) const override final
+			{
+				Serialization::Save(InOutOffset, InputOutputSize, Data);
+			}
+
+			virtual void Evaluate(
+				ILayerInstance* Instance,
+				float* OutputBuffer,
+				const float* InputBuffer,
+				const uint32 BatchSize,
+				const uint32 OutputBufferSize,
+				const uint32 InputBufferSize,
+				const uint32 OutputBufferStride,
+				const uint32 InputBufferStride) override final
+			{
+				NNE_RUNTIME_BASIC_TRACE_SCOPE(NNE::RuntimeBasic::Private::FSoftmaxLayer::Evaluate);
+				check(OutputBufferSize == GetOutputSize() && InputBufferSize == GetInputSize());
+				check(OutputBufferStride >= GetOutputSize() && InputBufferStride >= GetInputSize());
+				check(Instance == nullptr);
+				OperatorNanCheck(InputBuffer, BatchSize, InputBufferSize, InputBufferStride);
+
+				OperatorSoftmax(
+					OutputBuffer,
+					InputBuffer,
+					BatchSize,
+					InputOutputSize,
+					OutputBufferStride,
+					InputBufferStride);
+
+				OperatorNanCheck(OutputBuffer, BatchSize, OutputBufferSize, OutputBufferStride);
+			}
+
+			uint32 InputOutputSize = 0;
+		};
+
+
 		//--------------------------------------------------------------------------
 		// Layer Serialization
 		//--------------------------------------------------------------------------
@@ -5511,6 +5610,7 @@ namespace UE::NNE::RuntimeBasic
 					case ELayerType::GRUCell: OutLayer = MakeShared<FGRUCellLayer>(); break;
 					case ELayerType::Conv1d: OutLayer = MakeShared<FConv1dLayer>(); break;
 					case ELayerType::Conv2d: OutLayer = MakeShared<FConv2dLayer>(); break;
+					case ELayerType::Softmax: OutLayer = MakeShared<FSoftmaxLayer>(); break;
 					default: checkf(false, TEXT("Unknown Layer Id %i"), LayerTypeId);
 					}
 				}
@@ -5947,6 +6047,13 @@ namespace UE::NNE::RuntimeBasic
 		return StaticCastSharedPtr<Private::ILayer>(Layer);
 	}
 
+	FModelBuilderElement FModelBuilder::MakeSoftmax(const uint32 InputOutputSize)
+	{
+		const TSharedPtr<Private::FSoftmaxLayer> Layer = MakeShared<Private::FSoftmaxLayer>();
+		Layer->InputOutputSize = InputOutputSize;
+		return StaticCastSharedPtr<Private::ILayer>(Layer);
+	}
+
 	FModelBuilderElement FModelBuilder::MakeCopy(const uint32 InputOutputSize)
 	{
 		const TSharedPtr<Private::FCopyLayer> Layer = MakeShared<Private::FCopyLayer>();
diff --git a/Engine/Plugins/Experimental/NNERuntimeBasicCpu/Source/NNERuntimeBasicCpu/Public/NNERuntimeBasicCpuBuilder.h b/Engine/Plugins/Experimental/NNERuntimeBasicCpu/Source/NNERuntimeBasicCpu/Public/NNERuntimeBasicCpuBuilder.h
--- a/Engine/Plugins/Experimental/NNERuntimeBasicCpu/Source/NNERuntimeBasicCpu/Public/NNERuntimeBasicCpuBuilder.h
+++ b/Engine/Plugins/Experimental/NNERuntimeBasicCpu/Source/NNERuntimeBasicCpu/Public/NNERuntimeBasicCpuBuilder.h
@@ -232,6 +232,9 @@ namespace UE::NNE::RuntimeBasic
 		/** Makes a TanH Activation Layer */
 		UE_API FModelBuilderElement MakeTanH(const uint32 InputOutputSize);
 
+		/** Makes a Softmax Layer */
+		UE_API FModelBuilderElement MakeSoftmax(const uint32 InputOutputSize);
+
 		/** Makes a Copy Layer */
 		UE_API FModelBuilderElement MakeCopy(const uint32 InputOutputSize);
 
